Import Xcode GPU Debugger Plugin for Mac

- This plugin works similarly to the RenderDoc plugin for Windows. Enable the plugin in your project settings, and then press Shift+E to trigger a GPU frame capture of the primary viewport. It will automatically open in Xcode once the capture is complete, at which point you can replay the capture, step through it, etc...

diff --git a/Engine/Plugins/Developer/XcodeGPUDebuggerPlugin/Source/XcodeGPUDebuggerPlugin/Private/SXcodeGPUDebuggerPluginEditorExtension.cpp b/Engine/Plugins/Developer/XcodeGPUDebuggerPlugin/Source/XcodeGPUDebuggerPlugin/Private/SXcodeGPUDebuggerPluginEditorExtension.cpp
new file mode 100644
--- /dev/null
+++ b/Engine/Plugins/Developer/XcodeGPUDebuggerPlugin/Source/XcodeGPUDebuggerPlugin/Private/SXcodeGPUDebuggerPluginEditorExtension.cpp
@@ -0,0 +1,175 @@
+// Copyright Epic Games, Inc. All Rights Reserved.
+
+#if WITH_EDITOR
+
+#include "SXcodeGPUDebuggerPluginEditorExtension.h"
+
+#include "XcodeGPUDebuggerPluginCommands.h"
+#include "XcodeGPUDebuggerPluginModule.h"
+#include "XcodeGPUDebuggerPluginStyle.h"
+
+#include "Editor/EditorEngine.h"
+#include "Editor/UnrealEd/Public/SEditorViewportToolBarMenu.h"
+#include "Editor/UnrealEd/Public/SViewportToolBarComboMenu.h"
+#include "Editor/UnrealEd/Public/Kismet2/DebuggerCommands.h"
+#include "EditorStyleSet.h"
+#include "Framework/Application/SlateApplication.h"
+#include "Framework/MultiBox/MultiBoxBuilder.h"
+#include "Framework/MultiBox/MultiBoxDefs.h"
+#include "Interfaces/IMainFrameModule.h"
+#include "LevelEditor.h"
+#include "Widgets/Input/SButton.h"
+#include "Widgets/Images/SImage.h"
+
+extern UNREALED_API UEditorEngine* GEditor;
+
+class SXcodeGPUDebuggerCaptureButton : public SViewportToolBar
+{
+public:
+	SLATE_BEGIN_ARGS(SXcodeGPUDebuggerCaptureButton) { }
+	SLATE_END_ARGS()
+
+	/** Widget constructor */
+	void Construct(const FArguments& Args)
+	{
+		FSlateIcon IconBrush = FSlateIcon(FXcodeGPUDebuggerPluginStyle::Get()->GetStyleSetName(), "XcodeGPUDebuggerPlugin.CaptureFrameIcon");
+
+		ChildSlot
+		[
+			SNew(SButton)
+			.HAlign(HAlign_Center)
+			.VAlign(VAlign_Bottom)
+			.ButtonStyle(FEditorStyle::Get(), "ViewportMenu.Button")
+			.ContentPadding(FMargin(1.0f))
+			.ToolTipText(FXcodeGPUDebuggerPluginCommands::Get().CaptureFrameCommand->GetDescription())
+			.OnClicked_Lambda([this]()
+				{
+					FPlayWorldCommands::GlobalPlayWorldActions->GetActionForCommand(FXcodeGPUDebuggerPluginCommands::Get().CaptureFrameCommand)->Execute();
+					return(FReply::Handled());
+				})
+			[
+				SNew(SImage)
+				.Image(IconBrush.GetIcon())
+			]
+		];
+	}
+};
+
+FXcodeGPUDebuggerPluginEditorExtension::FXcodeGPUDebuggerPluginEditorExtension(FXcodeGPUDebuggerPluginModule* ThePlugin)
+	: LoadedDelegateHandle()
+	, ToolbarExtension()
+	, ExtensionManager()
+	, ToolbarExtender()
+	, IsEditorInitialized(false)
+{
+	// Defer Level Editor UI extensions until Level Editor has been loaded:
+	if (FModuleManager::Get().IsModuleLoaded("LevelEditor"))
+	{
+		Initialize(ThePlugin);
+	}
+	else
+	{
+		FModuleManager::Get().OnModulesChanged().AddLambda([this, ThePlugin](FName name, EModuleChangeReason reason)
+		{
+			if ((name == "LevelEditor") && (reason == EModuleChangeReason::ModuleLoaded))
+			{
+				Initialize(ThePlugin);
+			}
+		});
+	}
+}
+
+FXcodeGPUDebuggerPluginEditorExtension::~FXcodeGPUDebuggerPluginEditorExtension()
+{
+	if (ExtensionManager.IsValid())
+	{
+		FXcodeGPUDebuggerPluginStyle::Shutdown();
+		FXcodeGPUDebuggerPluginCommands::Unregister();
+
+		ToolbarExtender->RemoveExtension(ToolbarExtension.ToSharedRef());
+
+		ExtensionManager->RemoveExtender(ToolbarExtender);
+	}
+	else
+	{
+		ExtensionManager.Reset();
+	}
+}
+
+void FXcodeGPUDebuggerPluginEditorExtension::Initialize(FXcodeGPUDebuggerPluginModule* ThePlugin)
+{
+	if (GUsingNullRHI)
+	{
+		return;
+	}
+
+	// The LoadModule request below will crash if running as an editor commandlet!
+	check(!IsRunningCommandlet());
+
+	FXcodeGPUDebuggerPluginStyle::Initialize();
+	FXcodeGPUDebuggerPluginCommands::Register();
+
+	FLevelEditorModule& LevelEditorModule = FModuleManager::LoadModuleChecked<FLevelEditorModule>("LevelEditor");
+	TSharedRef<FUICommandList> CommandBindings = LevelEditorModule.GetGlobalLevelEditorActions();
+	ExtensionManager = LevelEditorModule.GetToolBarExtensibilityManager();
+	ToolbarExtender = MakeShareable(new FExtender);
+	ToolbarExtension = ToolbarExtender->AddToolBarExtension("CameraSpeed", EExtensionHook::After, CommandBindings,
+		FToolBarExtensionDelegate::CreateLambda([this, ThePlugin](FToolBarBuilder& ToolbarBuilder)
+	{
+		AddToolbarExtension(ToolbarBuilder, ThePlugin); })
+	);
+	ExtensionManager->AddExtender(ToolbarExtender);
+
+	IsEditorInitialized = false;
+	FSlateRenderer* SlateRenderer = FSlateApplication::Get().GetRenderer();
+	LoadedDelegateHandle = SlateRenderer->OnSlateWindowRendered().AddRaw(this, &FXcodeGPUDebuggerPluginEditorExtension::OnEditorLoaded);
+}
+
+void FXcodeGPUDebuggerPluginEditorExtension::OnEditorLoaded(SWindow& SlateWindow, void* ViewportRHIPtr)
+{
+	// Would be nice to use the preprocessor definition WITH_EDITOR instead, but the user may launch a standalone the game through the editor...
+	if (GEditor == nullptr)
+	{
+		return;
+	}
+
+	if (IsInGameThread())
+	{
+		FSlateRenderer* SlateRenderer = FSlateApplication::Get().GetRenderer();
+		SlateRenderer->OnSlateWindowRendered().Remove(LoadedDelegateHandle);
+	}
+
+	if (IsEditorInitialized)
+	{
+		return;
+	}
+	IsEditorInitialized = true;
+
+	if(FPlayWorldCommands::GlobalPlayWorldActions.IsValid())
+	{
+		//Register the editor hotkeys
+		FPlayWorldCommands::GlobalPlayWorldActions->MapAction(FXcodeGPUDebuggerPluginCommands::Get().CaptureFrameCommand,
+			FExecuteAction::CreateLambda([]()
+			{
+				FXcodeGPUDebuggerPluginModule& PluginModule = FModuleManager::GetModuleChecked<FXcodeGPUDebuggerPluginModule>("XcodeGPUDebuggerPlugin");
+				PluginModule.CaptureFrame(nullptr, IRenderCaptureProvider::ECaptureFlags_Launch, FString());
+			}),
+			FCanExecuteAction());
+ 	}
+}
+
+void FXcodeGPUDebuggerPluginEditorExtension::AddToolbarExtension(FToolBarBuilder& ToolbarBuilder, FXcodeGPUDebuggerPluginModule* ThePlugin)
+{
+#define LOCTEXT_NAMESPACE "LevelEditorToolBar"
+
+	UE_LOG(XcodeGPUDebuggerPlugin, Log, TEXT("Attaching toolbar extension..."));
+	ToolbarBuilder.AddSeparator();
+
+	ToolbarBuilder.BeginSection("XcodeGPUDebuggerPlugin");
+	ToolbarBuilder.AddWidget(SNew(SXcodeGPUDebuggerCaptureButton));
+	ToolbarBuilder.EndSection();
+
+#undef LOCTEXT_NAMESPACE
+}
+
+#endif //WITH_EDITOR
diff --git a/Engine/Plugins/Developer/XcodeGPUDebuggerPlugin/Source/XcodeGPUDebuggerPlugin/Private/SXcodeGPUDebuggerPluginEditorExtension.h b/Engine/Plugins/Developer/XcodeGPUDebuggerPlugin/Source/XcodeGPUDebuggerPlugin/Private/SXcodeGPUDebuggerPluginEditorExtension.h
new file mode 100644
--- /dev/null
+++ b/Engine/Plugins/Developer/XcodeGPUDebuggerPlugin/Source/XcodeGPUDebuggerPlugin/Private/SXcodeGPUDebuggerPluginEditorExtension.h
@@ -0,0 +1,33 @@
+// Copyright Epic Games, Inc. All Rights Reserved.
+
+#pragma once
+
+#if WITH_EDITOR
+
+#include "Editor/UnrealEd/Public/SViewportToolBar.h"
+
+class FXcodeGPUDebuggerPluginModule;
+class FExtensionBase;
+class FExtensibilityManager;
+class FExtender;
+class FToolBarBuilder;
+
+class FXcodeGPUDebuggerPluginEditorExtension
+{
+public:
+	FXcodeGPUDebuggerPluginEditorExtension(FXcodeGPUDebuggerPluginModule* ThePlugin);
+	~FXcodeGPUDebuggerPluginEditorExtension();
+
+private:
+	void Initialize(FXcodeGPUDebuggerPluginModule* ThePlugin);
+	void OnEditorLoaded(SWindow& SlateWindow, void* ViewportRHIPtr);
+	void AddToolbarExtension(FToolBarBuilder& ToolbarBuilder, FXcodeGPUDebuggerPluginModule* ThePlugin);
+
+	FDelegateHandle LoadedDelegateHandle;
+	TSharedPtr<const FExtensionBase> ToolbarExtension;
+	TSharedPtr<FExtensibilityManager> ExtensionManager;
+	TSharedPtr<FExtender> ToolbarExtender;
+	bool IsEditorInitialized;
+};
+
+#endif // WITH_EDITOR
diff --git a/Engine/Plugins/Developer/XcodeGPUDebuggerPlugin/Source/XcodeGPUDebuggerPlugin/Private/XcodeGPUDebuggerPluginCommands.cpp b/Engine/Plugins/Developer/XcodeGPUDebuggerPlugin/Source/XcodeGPUDebuggerPlugin/Private/XcodeGPUDebuggerPluginCommands.cpp
new file mode 100644
--- /dev/null
+++ b/Engine/Plugins/Developer/XcodeGPUDebuggerPlugin/Source/XcodeGPUDebuggerPlugin/Private/XcodeGPUDebuggerPluginCommands.cpp
@@ -0,0 +1,16 @@
+// Copyright Epic Games, Inc. All Rights Reserved.
+
+#if WITH_EDITOR
+
+#include "XcodeGPUDebuggerPluginCommands.h"
+
+#define LOCTEXT_NAMESPACE "XcodeGPUDebuggerPlugin"
+
+void FXcodeGPUDebuggerPluginCommands::RegisterCommands()
+{
+	UI_COMMAND(CaptureFrameCommand, "Capture Frame", "Captures the next frame and launches Xcode (Shift+E)", EUserInterfaceActionType::Button, FInputChord(EKeys::E, EModifierKey::Shift));
+}
+
+#undef LOCTEXT_NAMESPACE
+
+#endif // WITH_EDITOR
diff --git a/Engine/Plugins/Developer/XcodeGPUDebuggerPlugin/Source/XcodeGPUDebuggerPlugin/Private/XcodeGPUDebuggerPluginCommands.h b/Engine/Plugins/Developer/XcodeGPUDebuggerPlugin/Source/XcodeGPUDebuggerPlugin/Private/XcodeGPUDebuggerPluginCommands.h
new file mode 100644
--- /dev/null
+++ b/Engine/Plugins/Developer/XcodeGPUDebuggerPlugin/Source/XcodeGPUDebuggerPlugin/Private/XcodeGPUDebuggerPluginCommands.h
@@ -0,0 +1,23 @@
+// Copyright Epic Games, Inc. All Rights Reserved.
+
+#pragma once
+
+#if WITH_EDITOR
+
+#include "XcodeGPUDebuggerPluginStyle.h"
+#include "Framework/Commands/Commands.h"
+
+class FXcodeGPUDebuggerPluginCommands : public TCommands<FXcodeGPUDebuggerPluginCommands>
+{
+public:
+	FXcodeGPUDebuggerPluginCommands()
+		: TCommands<FXcodeGPUDebuggerPluginCommands>(TEXT("XcodeGPUDebuggerPlugin"), NSLOCTEXT("Contexts", "XcodeGPUDebuggerPlugin", "XcodeGPUDebugger Plugin"), NAME_None, FXcodeGPUDebuggerPluginStyle::Get()->GetStyleSetName())
+	{
+	}
+
+	virtual void RegisterCommands() override;
+
+	TSharedPtr<class FUICommandInfo> CaptureFrameCommand;
+};
+
+#endif // WITH_EDITOR
diff --git a/Engine/Plugins/Developer/XcodeGPUDebuggerPlugin/Source/XcodeGPUDebuggerPlugin/Private/XcodeGPUDebuggerPluginModule.cpp b/Engine/Plugins/Developer/XcodeGPUDebuggerPlugin/Source/XcodeGPUDebuggerPlugin/Private/XcodeGPUDebuggerPluginModule.cpp
new file mode 100644
--- /dev/null
+++ b/Engine/Plugins/Developer/XcodeGPUDebuggerPlugin/Source/XcodeGPUDebuggerPlugin/Private/XcodeGPUDebuggerPluginModule.cpp
@@ -0,0 +1,382 @@
+// Copyright Epic Games, Inc. All Rights Reserved.
+
+#include "XcodeGPUDebuggerPluginModule.h"
+
+#include "Engine/Engine.h"
+#include "Engine/GameViewportClient.h"
+#include "Async/Async.h"
+#include "HAL/FileManager.h"
+#include "Internationalization/Internationalization.h"
+#include "Misc/ConfigCacheIni.h"
+#include "RendererInterface.h"
+#include "RenderingThread.h"
+#include "GeneralProjectSettings.h"
+
+#if WITH_EDITOR
+#include "Editor.h"
+#include "Editor/EditorEngine.h"
+#include "ISettingsModule.h"
+#include "ISettingsSection.h"
+#include "LevelEditor.h"
+#include "SXcodeGPUDebuggerPluginEditorExtension.h"
+#include "UnrealClient.h"
+extern UNREALED_API UEditorEngine* GEditor;
+#endif // WITH_EDITOR
+
+#include <Metal/MTLDevice.h>
+#include <Metal/MTLCaptureManager.h>
+
+DEFINE_LOG_CATEGORY(XcodeGPUDebuggerPlugin);
+
+#define LOCTEXT_NAMESPACE "XcodeGPUDebuggerPlugin"
+
+struct FXcodeGPUDebuggerAsyncGraphTask : public FAsyncGraphTaskBase
+{
+	ENamedThreads::Type TargetThread;
+	TFunction<void()> TheTask;
+
+	FXcodeGPUDebuggerAsyncGraphTask(ENamedThreads::Type Thread, TFunction<void()>&& Task) : TargetThread(Thread), TheTask(MoveTemp(Task)) { }
+	void DoTask(ENamedThreads::Type CurrentThread, const FGraphEventRef& MyCompletionGraphEvent) { TheTask(); }
+	ENamedThreads::Type GetDesiredThread() { return(TargetThread); }
+};
+
+class FXcodeGPUDebuggerFrameCapturer
+{
+public:
+	static FString MakeXcodeGPUDebuggerCaptureFilePath(FString const& InFileName)
+	{
+		if (InFileName.IsEmpty())
+		{
+			return InFileName;
+		}
+
+		const bool bAbsoluteFileName = !FPaths::IsRelative(InFileName);
+		FString FileName = bAbsoluteFileName ? InFileName : FPaths::ConvertRelativePathToFull(FPaths::ProjectSavedDir() / FString("XcodeGPUTraceCaptures") / InFileName);
+		FileName = FPaths::SetExtension(FileName, "gputrace");
+		FPaths::MakePlatformFilename(FileName);
+		return FileName;
+	}
+
+	static void BeginFrameCapture(const FString& CaptureFileName)
+	{
+		UE4_GEmitDrawEvents_BeforeCapture = GetEmitDrawEvents();
+		SetEmitDrawEvents(true);
+
+		id<MTLDevice> MetalDevice = (id<MTLDevice>)GDynamicRHI->RHIGetNativeDevice();
+
+		MTLCaptureManager* CaptureManager = [MTLCaptureManager sharedCaptureManager];
+
+		MTLCaptureDescriptor* CaptureDescriptor = [[MTLCaptureDescriptor alloc] init];
+		CaptureDescriptor.captureObject         = MetalDevice;
+
+#if PLATFORM_IOS || PLATFORM_TVOS
+		// On iOS/tvOS, assume that Xcode is attached.
+		CaptureDescriptor.destination =  MTLCaptureDestinationDeveloperTools;
+#else
+		FString SaveURL = FXcodeGPUDebuggerFrameCapturer::MakeXcodeGPUDebuggerCaptureFilePath(CaptureFileName);
+		FTCHARToUTF8 TCUrlStr(*SaveURL);
+		NSString* NSUrlStr = [NSString stringWithUTF8String:TCUrlStr.Get()];
+		NSURL* Url = [NSURL fileURLWithPath:NSUrlStr];
+
+		CaptureDescriptor.destination 			= MTLCaptureDestinationGPUTraceDocument;
+		CaptureDescriptor.outputURL 			= Url;
+#endif
+
+		NSError* Error;
+		if (![CaptureManager startCaptureWithDescriptor:CaptureDescriptor error:&Error])
+		{
+			FString ErrorString([Error localizedDescription]);
+			UE_LOG(XcodeGPUDebuggerPlugin, Warning, TEXT("Failed to start capture (%s)"), *ErrorString);
+		}
+	}
+
+	static void EndFrameCapture()
+	{
+		FRHICommandListExecutor::GetImmediateCommandList().SubmitCommandsAndFlushGPU();
+
+		MTLCaptureManager* CaptureManager = [MTLCaptureManager sharedCaptureManager];
+		[CaptureManager stopCapture];
+
+		SetEmitDrawEvents(UE4_GEmitDrawEvents_BeforeCapture);
+	}
+
+	static void SaveAndLaunch(FXcodeGPUDebuggerPluginModule* Plugin, int32 Flags, const FString& InDestPath)
+	{
+#if !(PLATFORM_IOS || PLATFORM_TVOS)
+		FString DestPath = MakeXcodeGPUDebuggerCaptureFilePath(InDestPath);
+
+		if (Flags & IRenderCaptureProvider::ECaptureFlags_Launch)
+		{
+			TGraphTask<FXcodeGPUDebuggerAsyncGraphTask>::CreateTask().ConstructAndDispatchWhenReady(ENamedThreads::GameThread, [DestPath]()
+			{
+				NSString* GPUTracePath = [NSString stringWithFString:DestPath];
+				[[NSWorkspace sharedWorkspace] openFile:GPUTracePath withApplication:@"Xcode"];
+			});
+		}
+#endif
+	}
+
+private:
+	static bool UE4_GEmitDrawEvents_BeforeCapture;
+};
+
+bool FXcodeGPUDebuggerFrameCapturer::UE4_GEmitDrawEvents_BeforeCapture = false;
+
+class FXcodeGPUDebuggerDummyInputDevice : public IInputDevice
+{
+public:
+	FXcodeGPUDebuggerDummyInputDevice(FXcodeGPUDebuggerPluginModule* InPlugin) : ThePlugin(InPlugin) { }
+	virtual ~FXcodeGPUDebuggerDummyInputDevice() { }
+
+	/** Tick the interface (used for controlling full engine frame captures). */
+	virtual void Tick(float DeltaTime) override
+	{
+		check(ThePlugin);
+		ThePlugin->Tick(DeltaTime);
+	}
+
+	/** The remaining interfaces are irrelevant for this dummy input device. */
+	virtual void SendControllerEvents() override { }
+	virtual void SetMessageHandler(const TSharedRef<FGenericApplicationMessageHandler>& InMessageHandler) override { }
+	virtual bool Exec(UWorld* InWorld, const TCHAR* Cmd, FOutputDevice& Ar) override { return(false); }
+	virtual void SetChannelValue(int32 ControllerId, FForceFeedbackChannelType ChannelType, float Value) override { }
+	virtual void SetChannelValues(int32 ControllerId, const FForceFeedbackValues &values) override { }
+
+private:
+	FXcodeGPUDebuggerPluginModule* ThePlugin;
+};
+
+TSharedPtr<class IInputDevice> FXcodeGPUDebuggerPluginModule::CreateInputDevice(const TSharedRef<FGenericApplicationMessageHandler>& InMessageHandler)
+{
+	UE_LOG(XcodeGPUDebuggerPlugin, Log, TEXT("Creating dummy input device (for intercepting engine ticks)"));
+	FXcodeGPUDebuggerDummyInputDevice* InputDev = new FXcodeGPUDebuggerDummyInputDevice(this);
+	return MakeShareable((IInputDevice*)InputDev);
+}
+
+void FXcodeGPUDebuggerPluginModule::StartupModule()
+{
+#if !UE_BUILD_SHIPPING // Disable in shipping builds
+#if WITH_EDITOR
+	EditorExtensions = nullptr;
+#endif // WITH_EDITOR
+
+	// Enable Metal capture layer.
+	setenv("METAL_CAPTURE_ENABLED", "1", 1);
+
+	InjectDebugExecKeybind();
+
+	// Do not load the plugin if UE4 is headless.
+	if (GUsingNullRHI)
+	{
+		UE_LOG(XcodeGPUDebuggerPlugin, Display, TEXT("Xcode GPU Debugger plugin will not be loaded (headless UE4 instance)."));
+		return;
+	}
+
+	IModularFeatures::Get().RegisterModularFeature(IRenderCaptureProvider::GetModularFeatureName(), (IRenderCaptureProvider*)this);
+	IModularFeatures::Get().RegisterModularFeature(IInputDeviceModule::GetModularFeatureName(), (IInputDeviceModule*)this);
+
+	bCaptureInProgress = false;
+
+	// Setup XcodeGPUDebugger settings
+	FString XcodeGPUDebuggerCapturePath = FPaths::ProjectSavedDir() / TEXT("XcodeGPUTraceCaptures");
+	if (!IFileManager::Get().DirectoryExists(*XcodeGPUDebuggerCapturePath))
+	{
+		IFileManager::Get().MakeDirectory(*XcodeGPUDebuggerCapturePath, true);
+	}
+
+	// Register console commands.
+	static FAutoConsoleCommand CCmdXcodeGPUDebuggerCaptureFrame = FAutoConsoleCommand(
+		TEXT("Xcode.CaptureFrame"),
+		TEXT("Captures the rendering commands of the next frame and launches Xcode"),
+		FConsoleCommandDelegate::CreateRaw(this, &FXcodeGPUDebuggerPluginModule::CaptureFrame)
+	);
+
+#if WITH_EDITOR
+	EditorExtensions = new FXcodeGPUDebuggerPluginEditorExtension(this);
+#endif // WITH_EDITOR
+
+	UE_LOG(XcodeGPUDebuggerPlugin, Log, TEXT("Xcode GPU Debugger plugin is ready!"));
+#endif // !UE_BUILD_SHIPPING
+}
+
+void FXcodeGPUDebuggerPluginModule::BeginFrameCapture(const FString& InCaptureFileName)
+{
+	UE_LOG(XcodeGPUDebuggerPlugin, Log, TEXT("Capturing the frame and running Xcode (to file '%s.gputrace')"), *InCaptureFileName);
+
+	ENQUEUE_RENDER_COMMAND(StartXcodeGPUDebuggerCapture)(
+		[InCaptureFileName](FRHICommandListImmediate& RHICmdList)
+		{
+			FXcodeGPUDebuggerFrameCapturer::BeginFrameCapture(InCaptureFileName);
+		}
+    );
+}
+
+void FXcodeGPUDebuggerPluginModule::InjectDebugExecKeybind()
+{
+	FConfigSection* Section = GConfig->GetSectionPrivate(TEXT("/Script/Engine.PlayerInput"), false, false, GInputIni);
+	if (Section != nullptr)
+	{
+		Section->HandleAddCommand(TEXT("DebugExecBindings"), TEXT("(Key=E,Command=\"Xcode.CaptureFrame\", Shift=true)"), false);
+	}
+}
+
+void FXcodeGPUDebuggerPluginModule::EndFrameCapture(void* HWnd, uint32 Flags, const FString& DestFileName)
+{
+	ENQUEUE_RENDER_COMMAND(EndXcodeGPUDebuggerCapture)(
+		[Plugin = this, Flags, DestFileName](FRHICommandListImmediate& RHICmdList)
+		{
+			FXcodeGPUDebuggerFrameCapturer::EndFrameCapture();
+			FXcodeGPUDebuggerFrameCapturer::SaveAndLaunch(Plugin, Flags, DestFileName);
+		});
+
+	bCaptureInProgress = false;
+}
+
+void FXcodeGPUDebuggerPluginModule::CaptureFrame(FViewport* InViewport, uint32 InFlags, FString const& InDestFileName)
+{
+	check(IsInGameThread());
+
+	// Don't do anything if we're currently already waiting for a capture to end.
+	if (bCaptureInProgress)
+	{
+		return;
+	}
+
+    DoFrameCaptureCurrentViewport(InViewport, InFlags, InDestFileName);
+}
+
+void FXcodeGPUDebuggerPluginModule::CaptureFrame()
+{
+	CaptureFrame(nullptr, ECaptureFlags_Launch, FString());
+}
+
+void FXcodeGPUDebuggerPluginModule::DoFrameCaptureCurrentViewport(FViewport* InViewport, uint32 InFlags, const FString& InDestFileName)
+{
+	CaptureFileName = InDestFileName;
+
+	// Use the current date/time as a fallback.
+	if (CaptureFileName.IsEmpty())
+	{
+		CaptureFileName = FDateTime::Now().ToString();
+	}
+
+	BeginFrameCapture(CaptureFileName);
+
+	// infer the intended viewport to intercept/capture:
+	FViewport* Viewport = InViewport;
+
+	check(GEngine);
+	if (!Viewport && GEngine->GameViewport)
+	{
+		check(GEngine->GameViewport->Viewport);
+		if (GEngine->GameViewport->Viewport->HasFocus())
+		{
+			Viewport = GEngine->GameViewport->Viewport;
+		}
+	}
+
+#if WITH_EDITOR
+	if (!Viewport && GEditor)
+	{
+		// WARNING: capturing from a "PIE-Eject" Editor viewport will not work as
+		// expected; in such case, capture via the console command
+		// (this has something to do with the 'active' editor viewport when the UI
+		// button is clicked versus the one which the console is attached to)
+		Viewport = GEditor->GetActiveViewport();
+	}
+#endif // WITH_EDITOR
+
+	check(Viewport);
+	Viewport->Draw(true);
+
+	EndFrameCapture(Viewport->GetWindow(), InFlags, CaptureFileName);
+}
+
+void FXcodeGPUDebuggerPluginModule::Tick(float DeltaTime)
+{
+	if (bCaptureInProgress)
+	{
+		EndFrameCapture(nullptr, ECaptureFlags_Launch, FString());
+	}
+}
+
+void FXcodeGPUDebuggerPluginModule::StartXcode(FString LaunchPath)
+{
+	if (LaunchPath.IsEmpty())
+	{
+		return;
+	}
+
+	NSString* GPUTracePath = [NSString stringWithFString:LaunchPath];
+	[[NSWorkspace sharedWorkspace] openFile:GPUTracePath withApplication:@"Xcode"];
+}
+
+void FXcodeGPUDebuggerPluginModule::ShutdownModule()
+{
+    if (GUsingNullRHI)
+    {
+        return;
+    }
+
+	IModularFeatures::Get().UnregisterModularFeature(IRenderCaptureProvider::GetModularFeatureName(), (IRenderCaptureProvider*)this);
+
+#if WITH_EDITOR
+	delete EditorExtensions;
+#endif // WITH_EDITOR
+}
+
+void FXcodeGPUDebuggerPluginModule::BeginCapture(FRHICommandListImmediate* InRHICommandList, uint32 InFlags, FString const& InDestFileName)
+{
+	CaptureFlags = InFlags;
+	CaptureFileName = InDestFileName;
+
+    id<MTLDevice> MetalDevice = (id<MTLDevice>)GDynamicRHI->RHIGetNativeDevice();
+
+	InRHICommandList->SubmitCommandsAndFlushGPU();
+	InRHICommandList->EnqueueLambda([FileName = CaptureFileName, MetalDevice](FRHICommandListImmediate& RHICommandList)
+	{
+		MTLCaptureManager* CaptureManager = [MTLCaptureManager sharedCaptureManager];
+
+		MTLCaptureDescriptor* CaptureDescriptor = [[MTLCaptureDescriptor alloc] init];
+		CaptureDescriptor.captureObject         = MetalDevice;
+
+#if PLATFORM_IOS || PLATFORM_TVOS
+		// On iOS/tvOS, assume that Xcode is attached.
+		CaptureDescriptor.destination =  MTLCaptureDestinationDeveloperTools;
+#else
+		FString SaveURL = FXcodeGPUDebuggerFrameCapturer::MakeXcodeGPUDebuggerCaptureFilePath(FileName);
+		FTCHARToUTF8 TCUrlStr(*SaveURL);
+		NSString* NSUrlStr = [NSString stringWithUTF8String:TCUrlStr.Get()];
+		NSURL* Url = [NSURL fileURLWithPath:NSUrlStr];
+
+		CaptureDescriptor.destination =  MTLCaptureDestinationGPUTraceDocument;
+		CaptureDescriptor.outputURL = Url;
+#endif
+
+		NSError* Error;
+		if (![CaptureManager startCaptureWithDescriptor:CaptureDescriptor error:&Error])
+		{
+			FString ErrorString([Error localizedDescription]);
+			UE_LOG(XcodeGPUDebuggerPlugin, Warning, TEXT("Failed to start capture (%s)"), *ErrorString);
+		}
+	});
+}
+
+void FXcodeGPUDebuggerPluginModule::EndCapture(FRHICommandListImmediate* InRHICommandList)
+{
+	InRHICommandList->SubmitCommandsAndFlushGPU();
+    InRHICommandList->EnqueueLambda([this](FRHICommandListImmediate& RHICommandList)
+	{
+		MTLCaptureManager* CaptureManager = [MTLCaptureManager sharedCaptureManager];
+		[CaptureManager stopCapture];
+
+		FXcodeGPUDebuggerFrameCapturer::SaveAndLaunch(this, CaptureFlags, CaptureFileName);
+	});
+
+	CaptureFlags = 0;
+	CaptureFileName.Reset();
+}
+
+#undef LOCTEXT_NAMESPACE
+
+IMPLEMENT_MODULE(FXcodeGPUDebuggerPluginModule, XcodeGPUDebuggerPlugin)
diff --git a/Engine/Plugins/Developer/XcodeGPUDebuggerPlugin/Source/XcodeGPUDebuggerPlugin/Private/XcodeGPUDebuggerPluginModule.h b/Engine/Plugins/Developer/XcodeGPUDebuggerPlugin/Source/XcodeGPUDebuggerPlugin/Private/XcodeGPUDebuggerPluginModule.h
new file mode 100644
--- /dev/null
+++ b/Engine/Plugins/Developer/XcodeGPUDebuggerPlugin/Source/XcodeGPUDebuggerPlugin/Private/XcodeGPUDebuggerPluginModule.h
@@ -0,0 +1,58 @@
+// Copyright Epic Games, Inc. All Rights Reserved.
+
+#pragma once
+
+#include "IXcodeGPUDebuggerPlugin.h"
+
+DECLARE_LOG_CATEGORY_EXTERN(XcodeGPUDebuggerPlugin, Log, All);
+
+class FXcodeGPUDebuggerPluginEditorExtension;
+
+class FXcodeGPUDebuggerPluginModule : public IXcodeGPUDebuggerPlugin
+{
+	friend class FXcodeGPUDebuggerDummyInputDevice;
+	friend class FXcodeGPUDebuggerFrameCapturer;
+
+public:
+	// Begin IRenderCaptureProvider interface.
+	virtual void CaptureFrame(FViewport* InViewport, uint32 InFlags, FString const& InDestFileName) override;
+	virtual void BeginCapture(FRHICommandListImmediate* InRHICommandList, uint32 InFlags, FString const& InDestFileName) override;
+	virtual void EndCapture(FRHICommandListImmediate* InRHICommandList) override;
+	// End IRenderCaptureProvider interface.
+
+protected:
+	// Begin IModuleInterface interface.
+	virtual void StartupModule() override;
+	virtual void ShutdownModule() override;
+	// End IModuleInterface interface.
+
+	// Begin IInputDeviceModule interface.
+	virtual TSharedPtr<class IInputDevice> CreateInputDevice(const TSharedRef<FGenericApplicationMessageHandler>& InMessageHandler) override;
+	// End IInputDeviceModule interface.
+
+private:
+	void Tick(float DeltaTime);
+
+	/** Injects a debug key bind into the local player so that the hot key works the same in game */
+	void InjectDebugExecKeybind();
+
+	/** Helper function for CVar command binding. */
+	void CaptureFrame();
+
+	void DoFrameCaptureCurrentViewport(FViewport* InViewport, uint32 InFlags, FString const& InDestFileName);
+
+	void BeginFrameCapture(const FString& InCaptureFileName);
+	void EndFrameCapture(void* HWnd, uint32 Flags, FString const& DestFileName);
+
+	void StartXcode(FString CapturePath);
+
+private:
+    uint32 CaptureFlags;
+    uint32 bCaptureInProgress:1;
+
+	FString CaptureFileName;
+
+#if WITH_EDITOR
+	FXcodeGPUDebuggerPluginEditorExtension* EditorExtensions;
+#endif // WITH_EDITOR
+};
diff --git a/Engine/Plugins/Developer/XcodeGPUDebuggerPlugin/Source/XcodeGPUDebuggerPlugin/Private/XcodeGPUDebuggerPluginStyle.cpp b/Engine/Plugins/Developer/XcodeGPUDebuggerPlugin/Source/XcodeGPUDebuggerPlugin/Private/XcodeGPUDebuggerPluginStyle.cpp
new file mode 100644
--- /dev/null
+++ b/Engine/Plugins/Developer/XcodeGPUDebuggerPlugin/Source/XcodeGPUDebuggerPlugin/Private/XcodeGPUDebuggerPluginStyle.cpp
@@ -0,0 +1,62 @@
+// Copyright Epic Games, Inc. All Rights Reserved.
+
+#if WITH_EDITOR
+
+#include "XcodeGPUDebuggerPluginStyle.h"
+#include "Styling/SlateStyle.h"
+#include "Styling/SlateStyleRegistry.h"
+#include "HAL/FileManager.h"
+#include "Interfaces/IPluginManager.h"
+
+FString FXcodeGPUDebuggerPluginStyle::InContent(const FString& RelativePath, const ANSICHAR* Extension)
+{
+	auto myself = IPluginManager::Get().FindPlugin(TEXT("XcodeGPUDebuggerPlugin"));
+	check(myself.IsValid());
+	static FString ContentDir = myself->GetBaseDir() / TEXT("Resources");
+	return (ContentDir / RelativePath) + Extension;
+}
+
+TSharedPtr<FSlateStyleSet> FXcodeGPUDebuggerPluginStyle::StyleSet = nullptr;
+TSharedPtr<class ISlateStyle> FXcodeGPUDebuggerPluginStyle::Get() { return StyleSet; }
+
+void FXcodeGPUDebuggerPluginStyle::Initialize()
+{
+	// Only register once
+	if (StyleSet.IsValid())
+	{
+		return;
+	}
+
+	StyleSet = MakeShareable(new FSlateStyleSet("XcodeGPUDebuggerPluginStyle"));
+
+	FString ProjectResourceDir = FPaths::ProjectPluginsDir() / TEXT("XcodeGPUDebuggerPlugin/Resources");
+	FString EngineResourceDir = FPaths::EnginePluginsDir() / TEXT("XcodeGPUDebuggerPlugin/Resources");
+
+	if (IFileManager::Get().DirectoryExists(*ProjectResourceDir)) //Is the plugin in the project? In that case, use those resources
+	{
+		StyleSet->SetContentRoot(ProjectResourceDir);
+		StyleSet->SetCoreContentRoot(ProjectResourceDir);
+	}
+	else //Otherwise, use the global ones
+	{
+		StyleSet->SetContentRoot(EngineResourceDir);
+		StyleSet->SetCoreContentRoot(EngineResourceDir);
+	}
+
+	StyleSet->Set("XcodeGPUDebuggerPlugin.Icon", new FSlateImageBrush(FXcodeGPUDebuggerPluginStyle::InContent("Icon128", ".png"), FVector2D(128.0f, 128.0f)));
+	StyleSet->Set("XcodeGPUDebuggerPlugin.CaptureFrameIcon", new FSlateImageBrush(FXcodeGPUDebuggerPluginStyle::InContent("ViewportIcon16", ".png"), FVector2D(16.0f, 16.0f)));
+
+	FSlateStyleRegistry::RegisterSlateStyle(*StyleSet.Get());
+};
+
+void FXcodeGPUDebuggerPluginStyle::Shutdown()
+{
+	if (StyleSet.IsValid())
+	{
+		FSlateStyleRegistry::UnRegisterSlateStyle(*StyleSet.Get());
+		ensure(StyleSet.IsUnique());
+		StyleSet.Reset();
+	}
+}
+
+#endif //WITH_EDITOR
diff --git a/Engine/Plugins/Developer/XcodeGPUDebuggerPlugin/Source/XcodeGPUDebuggerPlugin/Private/XcodeGPUDebuggerPluginStyle.h b/Engine/Plugins/Developer/XcodeGPUDebuggerPlugin/Source/XcodeGPUDebuggerPlugin/Private/XcodeGPUDebuggerPluginStyle.h
new file mode 100644
--- /dev/null
+++ b/Engine/Plugins/Developer/XcodeGPUDebuggerPlugin/Source/XcodeGPUDebuggerPlugin/Private/XcodeGPUDebuggerPluginStyle.h
@@ -0,0 +1,23 @@
+// Copyright Epic Games, Inc. All Rights Reserved.
+
+#pragma once
+
+#if WITH_EDITOR
+
+#include "Styling/ISlateStyle.h"
+
+class FXcodeGPUDebuggerPluginStyle
+{
+public:
+	static void Initialize();
+	static void Shutdown();
+	static TSharedPtr<class ISlateStyle> Get();
+
+private:
+	static FString InContent(const FString& RelativePath, const ANSICHAR* Extension);
+
+private:
+	static TSharedPtr<class FSlateStyleSet> StyleSet;
+};
+
+#endif // WITH_EDITOR
diff --git a/Engine/Plugins/Developer/XcodeGPUDebuggerPlugin/Source/XcodeGPUDebuggerPlugin/Public/IXcodeGPUDebuggerPlugin.h b/Engine/Plugins/Developer/XcodeGPUDebuggerPlugin/Source/XcodeGPUDebuggerPlugin/Public/IXcodeGPUDebuggerPlugin.h
new file mode 100644
--- /dev/null
+++ b/Engine/Plugins/Developer/XcodeGPUDebuggerPlugin/Source/XcodeGPUDebuggerPlugin/Public/IXcodeGPUDebuggerPlugin.h
@@ -0,0 +1,20 @@
+// Copyright Epic Games, Inc. All Rights Reserved.
+
+#pragma once
+
+#include "IInputDeviceModule.h"
+#include "IRenderCaptureProvider.h"
+
+class IXcodeGPUDebuggerPlugin : public IInputDeviceModule, public IRenderCaptureProvider
+{
+public:
+	static inline IXcodeGPUDebuggerPlugin& Get()
+	{
+		return FModuleManager::LoadModuleChecked<IXcodeGPUDebuggerPlugin>("XcodeGPUDebugger");
+	}
+
+	static inline bool IsAvailable()
+	{
+		return FModuleManager::Get().IsModuleLoaded("XcodeGPUDebugger");
+	}
+};
diff --git a/Engine/Plugins/Developer/XcodeGPUDebuggerPlugin/Source/XcodeGPUDebuggerPlugin/XcodeGPUDebuggerPlugin.Build.cs b/Engine/Plugins/Developer/XcodeGPUDebuggerPlugin/Source/XcodeGPUDebuggerPlugin/XcodeGPUDebuggerPlugin.Build.cs
new file mode 100644
--- /dev/null
+++ b/Engine/Plugins/Developer/XcodeGPUDebuggerPlugin/Source/XcodeGPUDebuggerPlugin/XcodeGPUDebuggerPlugin.Build.cs
@@ -0,0 +1,43 @@
+// Copyright Epic Games, Inc. All Rights Reserved.
+
+using System.IO;
+
+namespace UnrealBuildTool.Rules
+{
+	public class XcodeGPUDebuggerPlugin : ModuleRules
+	{
+		public XcodeGPUDebuggerPlugin(ReadOnlyTargetRules Target) : base(Target)
+		{
+			PrivateDependencyModuleNames.AddRange(new string[] { });
+			PublicDependencyModuleNames.AddRange(new string[]
+			{
+				  "Core"
+				, "CoreUObject"
+				, "Engine"
+				, "InputCore"
+				, "DesktopPlatform"
+				, "Projects"
+				, "RenderCore"
+				, "InputDevice"
+				, "RHI"
+				, "DeveloperSettings"
+			});
+
+			PublicWeakFrameworks.Add("Metal");
+
+			if (Target.bBuildEditor == true)
+			{
+				DynamicallyLoadedModuleNames.AddRange(new string[] { "LevelEditor" });
+				PublicDependencyModuleNames.AddRange(new string[]
+				{
+					  "Slate"
+					, "SlateCore"
+					, "EditorStyle"
+					, "UnrealEd"
+					, "MainFrame"
+					, "GameProjectGeneration"
+				});
+			}
+		}
+	}
+}
diff --git a/Engine/Plugins/Developer/XcodeGPUDebuggerPlugin/XcodeGPUDebuggerPlugin.uplugin b/Engine/Plugins/Developer/XcodeGPUDebuggerPlugin/XcodeGPUDebuggerPlugin.uplugin
new file mode 100644
--- /dev/null
+++ b/Engine/Plugins/Developer/XcodeGPUDebuggerPlugin/XcodeGPUDebuggerPlugin.uplugin
@@ -0,0 +1,24 @@
+{
+	"FileVersion" : 1,
+	"FriendlyName" : "Xcode GPU Debugger Plugin",
+	"Version" : 1,
+	"VersionName" : "1.0",
+	"CreatedBy" : "Apple, Inc.",
+	"CreatedByURL" : "https://apple.com/",
+	"DocsURL" : "",
+	"Description" : "Xcode GPU debugger integration.",
+	"Category" : "Rendering",
+	"EnabledByDefault" : false,
+	"Modules" :
+	[
+		{
+			"Name": "XcodeGPUDebuggerPlugin",
+			"Type": "DeveloperTool",
+			"LoadingPhase": "PostConfigInit",
+			"WhitelistPlatforms":
+            [
+				"Mac"
+			]
+		}
+	]
+}
