Fixed XCode GPU programmatic captures missing a file name.

diff --git a/Engine/Plugins/Developer/XcodeGPUDebuggerPlugin/Source/XcodeGPUDebuggerPlugin/Private/XcodeGPUDebuggerPluginModule.cpp b/Engine/Plugins/Developer/XcodeGPUDebuggerPlugin/Source/XcodeGPUDebuggerPlugin/Private/XcodeGPUDebuggerPluginModule.cpp
--- a/Engine/Plugins/Developer/XcodeGPUDebuggerPlugin/Source/XcodeGPUDebuggerPlugin/Private/XcodeGPUDebuggerPluginModule.cpp
+++ b/Engine/Plugins/Developer/XcodeGPUDebuggerPlugin/Source/XcodeGPUDebuggerPlugin/Private/XcodeGPUDebuggerPluginModule.cpp
@@ -330,6 +330,12 @@ void FXcodeGPUDebuggerPluginModule::BeginCapture(FRHICommandListImmediate* InRHI
 	CaptureFlags = InFlags;
 	CaptureFileName = InDestFileName;
 
+	// Use the current date/time as a fallback.
+	if (CaptureFileName.IsEmpty())
+	{
+		CaptureFileName = FDateTime::Now().ToString();
+	}
+	
     id<MTLDevice> MetalDevice = (id<MTLDevice>)GDynamicRHI->RHIGetNativeDevice();
 
 	InRHICommandList->SubmitCommandsAndFlushGPU();
