[XcodeGPUDebuggerPlugin]
* Speculative fix (added include) for mac compile error

diff --git a/Engine/Plugins/Developer/XcodeGPUDebuggerPlugin/Source/XcodeGPUDebuggerPlugin/Private/SXcodeGPUDebuggerPluginEditorExtension.cpp b/Engine/Plugins/Developer/XcodeGPUDebuggerPlugin/Source/XcodeGPUDebuggerPlugin/Private/SXcodeGPUDebuggerPluginEditorExtension.cpp
--- a/Engine/Plugins/Developer/XcodeGPUDebuggerPlugin/Source/XcodeGPUDebuggerPlugin/Private/SXcodeGPUDebuggerPluginEditorExtension.cpp
+++ b/Engine/Plugins/Developer/XcodeGPUDebuggerPlugin/Source/XcodeGPUDebuggerPlugin/Private/SXcodeGPUDebuggerPluginEditorExtension.cpp
@@ -12,6 +12,7 @@
 #include "Editor/UnrealEd/Public/SEditorViewportToolBarMenu.h"
 #include "Editor/UnrealEd/Public/SViewportToolBarComboMenu.h"
 #include "Editor/UnrealEd/Public/Kismet2/DebuggerCommands.h"
+#include "RHI.h"
 #include "Styling/AppStyle.h"
 #include "Framework/Application/SlateApplication.h"
 #include "Framework/MultiBox/MultiBoxBuilder.h"
