Apple: Update runtime to macos 11 and clean up warnings from doing that.

diff --git a/Engine/Plugins/Developer/XCodeSourceCodeAccess/Source/XCodeSourceCodeAccess/Private/XCodeSourceCodeAccessor.cpp b/Engine/Plugins/Developer/XCodeSourceCodeAccess/Source/XCodeSourceCodeAccess/Private/XCodeSourceCodeAccessor.cpp
--- a/Engine/Plugins/Developer/XCodeSourceCodeAccess/Source/XCodeSourceCodeAccess/Private/XCodeSourceCodeAccessor.cpp
+++ b/Engine/Plugins/Developer/XCodeSourceCodeAccess/Source/XCodeSourceCodeAccess/Private/XCodeSourceCodeAccessor.cpp
@@ -183,8 +183,13 @@ bool FXCodeSourceCodeAccessor::OpenFileAtLine(const FString& FullPath, int32 Lin
 			FString XcodePath = FPlatformMisc::GetXcodePath();
 			XcodePath.RemoveFromEnd(TEXT("/Contents/Developer"));
 			
-			NSString* ProjectPath = [ProjPath.GetNSString() stringByDeletingLastPathComponent];
-			[[NSWorkspace sharedWorkspace] openFile:ProjectPath withApplication:XcodePath.GetNSString() andDeactivate:YES];
+			NSURL* ProjectPath = [NSURL fileURLWithPath:[ProjPath.GetNSString() stringByDeletingLastPathComponent]];
+			[[NSWorkspace sharedWorkspace]
+				openURLs:[NSArray arrayWithObject: ProjectPath]
+				withApplicationAtURL:[NSURL URLWithString:XcodePath.GetNSString()]
+				configuration: [NSWorkspaceOpenConfiguration configuration]
+				completionHandler:nil
+			];
 			
 			NSAppleScript* AppleScript = nil;
 			
diff --git a/Engine/Plugins/Developer/XcodeGPUDebuggerPlugin/Source/XcodeGPUDebuggerPlugin/Private/XcodeGPUDebuggerPluginModule.cpp b/Engine/Plugins/Developer/XcodeGPUDebuggerPlugin/Source/XcodeGPUDebuggerPlugin/Private/XcodeGPUDebuggerPluginModule.cpp
--- a/Engine/Plugins/Developer/XcodeGPUDebuggerPlugin/Source/XcodeGPUDebuggerPlugin/Private/XcodeGPUDebuggerPluginModule.cpp
+++ b/Engine/Plugins/Developer/XcodeGPUDebuggerPlugin/Source/XcodeGPUDebuggerPlugin/Private/XcodeGPUDebuggerPluginModule.cpp
@@ -109,8 +109,17 @@ public:
 		{
 			TGraphTask<FXcodeGPUDebuggerAsyncGraphTask>::CreateTask().ConstructAndDispatchWhenReady(ENamedThreads::GameThread, [DestPath]()
 			{
-				NSString* GPUTracePath = [NSString stringWithFString:DestPath];
-				[[NSWorkspace sharedWorkspace] openFile:GPUTracePath withApplication:@"Xcode"];
+				FString XcodePath = FPlatformMisc::GetXcodePath();
+				XcodePath.RemoveFromEnd(TEXT("/Contents/Developer"));
+
+				NSURL* GPUTracePath = [NSURL fileURLWithPath:[NSString stringWithFString:DestPath]];
+
+				[[NSWorkspace sharedWorkspace]
+					openURLs:[NSArray arrayWithObject: GPUTracePath]
+					withApplicationAtURL:[NSURL URLWithString:XcodePath.GetNSString()]
+					configuration: [NSWorkspaceOpenConfiguration configuration]
+					completionHandler:nil
+				];
 			});
 		}
 #endif
@@ -307,8 +316,17 @@ void FXcodeGPUDebuggerPluginModule::StartXcode(FString LaunchPath)
 		return;
 	}
 
-	NSString* GPUTracePath = [NSString stringWithFString:LaunchPath];
-	[[NSWorkspace sharedWorkspace] openFile:GPUTracePath withApplication:@"Xcode"];
+	FString XcodePath = FPlatformMisc::GetXcodePath();
+	XcodePath.RemoveFromEnd(TEXT("/Contents/Developer"));
+
+	NSURL* GPUTracePath = [NSURL fileURLWithPath:[NSString stringWithFString:LaunchPath]];
+
+	[[NSWorkspace sharedWorkspace]
+		openURLs:[NSArray arrayWithObject: GPUTracePath]
+		withApplicationAtURL:[NSURL URLWithString:XcodePath.GetNSString()]
+		configuration: [NSWorkspaceOpenConfiguration configuration]
+		completionHandler:nil
+	];
 }
 
 void FXcodeGPUDebuggerPluginModule::ShutdownModule()
diff --git a/Engine/Source/Developer/StandaloneRenderer/Private/Mac/OpenGL/SlateOpenGLContext.cpp b/Engine/Source/Developer/StandaloneRenderer/Private/Mac/OpenGL/SlateOpenGLContext.cpp
--- a/Engine/Source/Developer/StandaloneRenderer/Private/Mac/OpenGL/SlateOpenGLContext.cpp
+++ b/Engine/Source/Developer/StandaloneRenderer/Private/Mac/OpenGL/SlateOpenGLContext.cpp
@@ -16,26 +16,13 @@
 
 void LockGLContext(NSOpenGLContext* Context)
 {
-	if (FPlatformMisc::IsRunningOnMavericks())
-	{
-		CGLLockContext([Context CGLContextObj]);
-	}
-	else
-	{
-		[Context lock];
-	}
+	[Context lock];
 }
 
 void UnlockGLContext(NSOpenGLContext* Context)
 {
-	if (FPlatformMisc::IsRunningOnMavericks())
-	{
-		CGLUnlockContext([Context CGLContextObj]);
-	}
-	else
-	{
-		[Context unlock];
-	}
+
+	[Context unlock];
 }
 
 @implementation FSlateOpenGLLayer
@@ -197,20 +184,8 @@ void FSlateOpenGLContext::Initialize(void* InWindow, const FSlateOpenGLContext*
 		}
 
 		MainThreadCall(^{
-			if (FPlatformMisc::IsRunningOnMavericks() && ([Window styleMask] & NSWindowStyleMaskTexturedBackground))
-			{
-				NSView* SuperView = [[Window contentView] superview];
-				[SuperView addSubview:View];
-				[SuperView setWantsLayer:YES];
-				[SuperView addSubview:[Window standardWindowButton:NSWindowCloseButton]];
-				[SuperView addSubview:[Window standardWindowButton:NSWindowMiniaturizeButton]];
-				[SuperView addSubview:[Window standardWindowButton:NSWindowZoomButton]];
-			}
-			else
-			{
-				[View setWantsLayer:YES];
-				[Window setContentView:View];
-			}
+			[View setWantsLayer:YES];
+			[Window setContentView:View];
 
 			[[Window standardWindowButton:NSWindowCloseButton] setAction:@selector(performClose:)];
 		}, NSDefaultRunLoopMode, true);
diff --git a/Engine/Source/Programs/UnrealBuildTool/Platform/Mac/MacToolChain.cs b/Engine/Source/Programs/UnrealBuildTool/Platform/Mac/MacToolChain.cs
--- a/Engine/Source/Programs/UnrealBuildTool/Platform/Mac/MacToolChain.cs
+++ b/Engine/Source/Programs/UnrealBuildTool/Platform/Mac/MacToolChain.cs
@@ -17,12 +17,12 @@ namespace UnrealBuildTool
 		/// <summary>
 		/// Which version of the Mac OS X to allow at run time
 		/// </summary>
-		public string MacOSVersion = "10.15";
+		public string MacOSVersion = "11.00";
 
 		/// <summary>
 		/// Minimum version of Mac OS X to actually run on, running on earlier versions will display the system minimum version error dialog and exit.
 		/// </summary>
-		public string MinMacOSVersion = "10.15.7";
+		public string MinMacOSVersion = "11.7.9";
 
 		/// <summary>
 		/// Constructor
diff --git a/Engine/Source/Programs/UnrealInsights/Private/Mac/MacUnrealInsightsMain.cpp b/Engine/Source/Programs/UnrealInsights/Private/Mac/MacUnrealInsightsMain.cpp
--- a/Engine/Source/Programs/UnrealInsights/Private/Mac/MacUnrealInsightsMain.cpp
+++ b/Engine/Source/Programs/UnrealInsights/Private/Mac/MacUnrealInsightsMain.cpp
@@ -37,15 +37,25 @@
 	}
 	else if ([[NSFileManager defaultManager]fileExistsAtPath:filename] )
 	{
-		NSURL* BundleURL = [[NSRunningApplication currentApplication]bundleURL];
+		NSURL* BundleURL = [[NSRunningApplication currentApplication] bundleURL];
 
-		NSDictionary* Configuration = [NSDictionary dictionaryWithObject : [NSArray arrayWithObject : filename] forKey : NSWorkspaceLaunchConfigurationArguments];
+		NSWorkspaceOpenConfiguration* Configuration = [NSWorkspaceOpenConfiguration configuration];
+		[Configuration setCreatesNewApplicationInstance:YES];
+		[Configuration setPromptsUserIfNeeded:YES];
+		[Configuration setArguments:[NSArray arrayWithObject: filename]];
 
-		NSError* Error = nil;
+		[[NSWorkspace sharedWorkspace]
+			openApplicationAtURL: BundleURL
+			configuration: Configuration
+			completionHandler:^(NSRunningApplication * _Nullable app, NSError * _Nullable error)
+			{
+				if (error) {
+					NSLog(@"Failed to run the app: %@", error.localizedDescription);
+				}
+			}
+		];
 
-		NSRunningApplication* NewInstance = [[NSWorkspace sharedWorkspace]launchApplicationAtURL:BundleURL options : (NSWorkspaceLaunchOptions)(NSWorkspaceLaunchAsync | NSWorkspaceLaunchNewInstance) configuration : Configuration error : &Error];
-
-		return (NewInstance != nil);
+		return YES;
 	}
 	else
 	{
diff --git a/Engine/Source/Runtime/ApplicationCore/Private/Mac/CocoaWindow.cpp b/Engine/Source/Runtime/ApplicationCore/Private/Mac/CocoaWindow.cpp
--- a/Engine/Source/Runtime/ApplicationCore/Private/Mac/CocoaWindow.cpp
+++ b/Engine/Source/Runtime/ApplicationCore/Private/Mac/CocoaWindow.cpp
@@ -47,14 +47,7 @@ NSString* NSPerformDragOperation = @"NSPerformDragOperation";
 - (NSRect)openGLFrame
 {
 	SCOPED_AUTORELEASE_POOL;
-	if ([self styleMask] & NSWindowStyleMaskTexturedBackground)
-	{
-		return [self frame];
-	}
-	else
-	{
-		return [[self contentView] frame];
-	}
+	return [[self contentView] frame];
 }
 
 - (NSView*)openGLView
diff --git a/Engine/Source/Runtime/ApplicationCore/Private/Mac/MacApplication.cpp b/Engine/Source/Runtime/ApplicationCore/Private/Mac/MacApplication.cpp
--- a/Engine/Source/Runtime/ApplicationCore/Private/Mac/MacApplication.cpp
+++ b/Engine/Source/Runtime/ApplicationCore/Private/Mac/MacApplication.cpp
@@ -1029,7 +1029,7 @@ void FMacApplication::ProcessMouseDownEvent(const FDeferredMacEvent& Event, TSha
 				{
 					MainThreadCall(^{ [WindowHandle performMiniaturize:nil]; }, NSDefaultRunLoopMode, true);
 				}
-				else if (!FPlatformMisc::IsRunningOnMavericks())
+				else
 				{
 					MainThreadCall(^{ [WindowHandle zoom:nil]; }, NSDefaultRunLoopMode, true);
 				}
diff --git a/Engine/Source/Runtime/ApplicationCore/Private/Mac/MacWindow.cpp b/Engine/Source/Runtime/ApplicationCore/Private/Mac/MacWindow.cpp
--- a/Engine/Source/Runtime/ApplicationCore/Private/Mac/MacWindow.cpp
+++ b/Engine/Source/Runtime/ApplicationCore/Private/Mac/MacWindow.cpp
@@ -176,11 +176,8 @@ void FMacWindow::Initialize( FMacApplication* const Application, const TSharedRe
 				{
 					[WindowHandle setCollectionBehavior: NSWindowCollectionBehaviorFullScreenAuxiliary|NSWindowCollectionBehaviorDefault|NSWindowCollectionBehaviorManaged|NSWindowCollectionBehaviorParticipatesInCycle];
 
-					if (!FPlatformMisc::IsRunningOnMavericks())
-					{
-						WindowHandle.titlebarAppearsTransparent = YES;
-						WindowHandle.titleVisibility = NSWindowTitleHidden;
-					}
+					WindowHandle.titlebarAppearsTransparent = YES;
+					WindowHandle.titleVisibility = NSWindowTitleHidden;
 				}
 
 				SetText(*Definition->Title);
diff --git a/Engine/Source/Runtime/Core/Private/Mac/MacPlatformMisc.cpp b/Engine/Source/Runtime/Core/Private/Mac/MacPlatformMisc.cpp
--- a/Engine/Source/Runtime/Core/Private/Mac/MacPlatformMisc.cpp
+++ b/Engine/Source/Runtime/Core/Private/Mac/MacPlatformMisc.cpp
@@ -249,7 +249,6 @@ struct FMacApplicationInfo
 		MacPlatformGetOSVersion(OSBuild);
 
 		OSXVersion = FMacPlatformMisc::GetNSOperatingSystemVersion();
-		RunningOnMavericks = OSXVersion.majorVersion == 10 && OSXVersion.minorVersion == 9;
 
 		XcodeVersion.majorVersion = XcodeVersion.minorVersion = XcodeVersion.patchVersion = 0;
 
@@ -487,7 +486,6 @@ struct FMacApplicationInfo
 	bool bIsUnattended;
 	bool bIsSandboxed;
 	bool RunningOnBattery;
-	bool RunningOnMavericks;
 	int32 PowerSourceNotification;
 	int32 NumCores;
 	int64 SystemLogSize;
@@ -1770,11 +1768,6 @@ bool FMacPlatformMisc::IsRunningOnBattery()
 	return GMacAppInfo.RunningOnBattery;
 }
 
-bool FMacPlatformMisc::IsRunningOnMavericks()
-{
-	return GMacAppInfo.RunningOnMavericks;
-}
-
 int32 FMacPlatformMisc::MacOSXVersionCompare(uint8 Major, uint8 Minor, uint8 Revision)
 {
 	uint8 TargetValues[3] = {Major, Minor, Revision};
diff --git a/Engine/Source/Runtime/Core/Private/Mac/MacPlatformProcess.cpp b/Engine/Source/Runtime/Core/Private/Mac/MacPlatformProcess.cpp
--- a/Engine/Source/Runtime/Core/Private/Mac/MacPlatformProcess.cpp
+++ b/Engine/Source/Runtime/Core/Private/Mac/MacPlatformProcess.cpp
@@ -316,7 +316,7 @@ FProcHandle FMacPlatformProcess::CreateProcInternal(const TCHAR* URL, const TCHA
 		if (![[NSFileManager defaultManager] fileExistsAtPath: nsProcessPath])
 		{
 			NSString* AppName = [[nsProcessPath lastPathComponent] stringByDeletingPathExtension];
-			nsProcessPath = [[NSWorkspace sharedWorkspace] fullPathForApplication:AppName];
+			nsProcessPath = [[[NSWorkspace sharedWorkspace] URLForApplicationWithBundleIdentifier:AppName] path];
 		}
 		
 		if ([[NSFileManager defaultManager] fileExistsAtPath: nsProcessPath])
@@ -1210,7 +1210,7 @@ bool FMacPlatformProcess::LaunchFileInDefaultExternalApplication( const TCHAR* F
 		// Xcode project is a special case where we don't open the project file itself, but the .xcodeproj folder containing it
 		FileToOpen = [FileToOpen stringByDeletingLastPathComponent];
 	}
-	bool Result = [[NSWorkspace sharedWorkspace] openFile: FileToOpen];
+	bool Result = [[NSWorkspace sharedWorkspace] openURL: [NSURL fileURLWithPath:FileToOpen]];
 	CFRelease( CFFileName );
 
 	return Result;
diff --git a/Engine/Source/Runtime/Core/Public/Mac/MacPlatformMisc.h b/Engine/Source/Runtime/Core/Public/Mac/MacPlatformMisc.h
--- a/Engine/Source/Runtime/Core/Public/Mac/MacPlatformMisc.h
+++ b/Engine/Source/Runtime/Core/Public/Mac/MacPlatformMisc.h
@@ -104,11 +104,6 @@ struct CORE_API FMacPlatformMisc : public FApplePlatformMisc
 
 	static IPlatformChunkInstall* GetPlatformChunkInstall();
 
-	/**
-	 * Returns whether the Mac OS X version is 10.9.x or not.
-	 */
-	static bool IsRunningOnMavericks();
-
 	/**
 	 * Returns if current < target returns 1, if current > target returns 1, else current == target and returns 0.
 	 */
diff --git a/Engine/Source/Runtime/Launch/Private/Mac/LaunchMac.cpp b/Engine/Source/Runtime/Launch/Private/Mac/LaunchMac.cpp
--- a/Engine/Source/Runtime/Launch/Private/Mac/LaunchMac.cpp
+++ b/Engine/Source/Runtime/Launch/Private/Mac/LaunchMac.cpp
@@ -124,13 +124,23 @@ static int32 MacOSVersionCompare(const NSOperatingSystemVersion& VersionA, const
 		
 		NSURL* BundleURL = [[NSRunningApplication currentApplication] bundleURL];
 		
-		NSDictionary* Configuration = [NSDictionary dictionaryWithObject: [NSArray arrayWithObject: ProjectName] forKey: NSWorkspaceLaunchConfigurationArguments];
-		
-		NSError* Error = nil;
-		
-		NSRunningApplication* NewInstance = [[NSWorkspace sharedWorkspace] launchApplicationAtURL:BundleURL options:(NSWorkspaceLaunchOptions)(NSWorkspaceLaunchAsync|NSWorkspaceLaunchNewInstance) configuration:Configuration error:&Error];
+		NSWorkspaceOpenConfiguration* Configuration = [NSWorkspaceOpenConfiguration configuration];
+		[Configuration setCreatesNewApplicationInstance:YES];
+		[Configuration setPromptsUserIfNeeded:YES];
+		[Configuration setArguments:[NSArray arrayWithObject: ProjectName]];
+
+		[[NSWorkspace sharedWorkspace]
+			openApplicationAtURL: BundleURL
+			configuration: Configuration
+				completionHandler:^(NSRunningApplication * _Nullable app, NSError * _Nullable error)
+				{
+					if (error) {
+						NSLog(@"Failed to run the app: %@", error.localizedDescription);
+					}
+				}
+		];
 		
-		return (NewInstance != nil);
+		return YES;
 	}
 	else
 	{
diff --git a/Engine/Source/Runtime/Portal/LauncherPlatform/Private/Mac/LauncherPlatformMac.cpp b/Engine/Source/Runtime/Portal/LauncherPlatform/Private/Mac/LauncherPlatformMac.cpp
--- a/Engine/Source/Runtime/Portal/LauncherPlatform/Private/Mac/LauncherPlatformMac.cpp
+++ b/Engine/Source/Runtime/Portal/LauncherPlatform/Private/Mac/LauncherPlatformMac.cpp
@@ -55,14 +55,7 @@ bool FLauncherPlatformMac::IsLauncherInstalled() const
 {
 	// Otherwise search for it...
 	NSWorkspace* Workspace = [NSWorkspace sharedWorkspace];
-	NSString* Path = [Workspace fullPathForApplication:@"Epic Games Launcher"];
-	if( Path )
-	{
-		return true;
-	}
-
-	// Otherwise search for the old Launcher...
-	Path = [Workspace fullPathForApplication:@"Unreal Engine"];
+	NSURL* Path = [Workspace URLForApplicationWithBundleIdentifier:@"com.epicgames.EpicGamesLauncher"];
 	if( Path )
 	{
 		return true;
