Made the toolbar extension log verbose in render capture plugins

diff --git a/Engine/Plugins/Developer/PixWinPlugin/Source/PixWinPlugin/Private/SPixWinPluginEditorExtension.cpp b/Engine/Plugins/Developer/PixWinPlugin/Source/PixWinPlugin/Private/SPixWinPluginEditorExtension.cpp
--- a/Engine/Plugins/Developer/PixWinPlugin/Source/PixWinPlugin/Private/SPixWinPluginEditorExtension.cpp
+++ b/Engine/Plugins/Developer/PixWinPlugin/Source/PixWinPlugin/Private/SPixWinPluginEditorExtension.cpp
@@ -131,7 +131,7 @@ void FPixWinPluginEditorExtension::AddToolbarExtension(FToolBarBuilder& ToolbarB
 {
 #define LOCTEXT_NAMESPACE "LevelEditorToolBar"
 
-	UE_LOG(PixWinPlugin, Log, TEXT("Attaching toolbar extension..."));
+	UE_LOG(PixWinPlugin, Verbose, TEXT("Attaching toolbar extension..."));
 	ToolbarBuilder.AddSeparator();
 
 	ToolbarBuilder.BeginSection("PixWinPlugin");
diff --git a/Engine/Plugins/Developer/RenderDocPlugin/Source/RenderDocPlugin/Private/SRenderDocPluginEditorExtension.cpp b/Engine/Plugins/Developer/RenderDocPlugin/Source/RenderDocPlugin/Private/SRenderDocPluginEditorExtension.cpp
--- a/Engine/Plugins/Developer/RenderDocPlugin/Source/RenderDocPlugin/Private/SRenderDocPluginEditorExtension.cpp
+++ b/Engine/Plugins/Developer/RenderDocPlugin/Source/RenderDocPlugin/Private/SRenderDocPluginEditorExtension.cpp
@@ -139,7 +139,7 @@ void FRenderDocPluginEditorExtension::AddToolbarExtension(FToolBarBuilder& Toolb
 {
 #define LOCTEXT_NAMESPACE "LevelEditorToolBar"
 
-	UE_LOG(RenderDocPlugin, Log, TEXT("Attaching toolbar extension..."));
+	UE_LOG(RenderDocPlugin, Verbose, TEXT("Attaching toolbar extension..."));
 	ToolbarBuilder.AddSeparator();
 
 	ToolbarBuilder.BeginSection("RenderdocPlugin");
diff --git a/Engine/Plugins/Developer/XcodeGPUDebuggerPlugin/Source/XcodeGPUDebuggerPlugin/Private/SXcodeGPUDebuggerPluginEditorExtension.cpp b/Engine/Plugins/Developer/XcodeGPUDebuggerPlugin/Source/XcodeGPUDebuggerPlugin/Private/SXcodeGPUDebuggerPluginEditorExtension.cpp
--- a/Engine/Plugins/Developer/XcodeGPUDebuggerPlugin/Source/XcodeGPUDebuggerPlugin/Private/SXcodeGPUDebuggerPluginEditorExtension.cpp
+++ b/Engine/Plugins/Developer/XcodeGPUDebuggerPlugin/Source/XcodeGPUDebuggerPlugin/Private/SXcodeGPUDebuggerPluginEditorExtension.cpp
@@ -163,7 +163,7 @@ void FXcodeGPUDebuggerPluginEditorExtension::AddToolbarExtension(FToolBarBuilder
 {
 #define LOCTEXT_NAMESPACE "LevelEditorToolBar"
 
-	UE_LOG(XcodeGPUDebuggerPlugin, Log, TEXT("Attaching toolbar extension..."));
+	UE_LOG(XcodeGPUDebuggerPlugin, Verbose, TEXT("Attaching toolbar extension..."));
 	ToolbarBuilder.AddSeparator();
 
 	ToolbarBuilder.BeginSection("XcodeGPUDebuggerPlugin");
