Convert the PixWIN and XcodeGPUDebugger buttons to the new viewport toolbar.

#rnx

diff --git a/Engine/Plugins/Developer/PixWinPlugin/Source/PixWinPlugin/PixWinPlugin.Build.cs b/Engine/Plugins/Developer/PixWinPlugin/Source/PixWinPlugin/PixWinPlugin.Build.cs
--- a/Engine/Plugins/Developer/PixWinPlugin/Source/PixWinPlugin/PixWinPlugin.Build.cs
+++ b/Engine/Plugins/Developer/PixWinPlugin/Source/PixWinPlugin/PixWinPlugin.Build.cs
@@ -39,6 +39,7 @@ namespace UnrealBuildTool.Rules
 					, "UnrealEd"
 					, "MainFrame"
 					, "GameProjectGeneration"
+					, "ToolMenus"
 				);
 			}
 		}
diff --git a/Engine/Plugins/Developer/PixWinPlugin/Source/PixWinPlugin/Private/PixWinPluginCommands.cpp b/Engine/Plugins/Developer/PixWinPlugin/Source/PixWinPlugin/Private/PixWinPluginCommands.cpp
--- a/Engine/Plugins/Developer/PixWinPlugin/Source/PixWinPlugin/Private/PixWinPluginCommands.cpp
+++ b/Engine/Plugins/Developer/PixWinPlugin/Source/PixWinPlugin/Private/PixWinPluginCommands.cpp
@@ -8,7 +8,7 @@
 
 void FPixWinPluginCommands::RegisterCommands()
 {
-	UI_COMMAND(CaptureFrameCommand, "Capture Frame", "Captures the next frame and launches the PIX UI if not attached (Alt+F12)", EUserInterfaceActionType::Button, FInputChord(EKeys::F12, EModifierKey::Alt));
+	UI_COMMAND(CaptureFrame, "Capture Frame", "Captures the next frame and launches the PIX UI if not attached (Alt+F12)", EUserInterfaceActionType::Button, FInputChord(EKeys::F12, EModifierKey::Alt));
 }
 
 #undef LOCTEXT_NAMESPACE
diff --git a/Engine/Plugins/Developer/PixWinPlugin/Source/PixWinPlugin/Private/PixWinPluginCommands.h b/Engine/Plugins/Developer/PixWinPlugin/Source/PixWinPlugin/Private/PixWinPluginCommands.h
--- a/Engine/Plugins/Developer/PixWinPlugin/Source/PixWinPlugin/Private/PixWinPluginCommands.h
+++ b/Engine/Plugins/Developer/PixWinPlugin/Source/PixWinPlugin/Private/PixWinPluginCommands.h
@@ -17,7 +17,7 @@ public:
 
 	virtual void RegisterCommands() override;
 
-	TSharedPtr<class FUICommandInfo> CaptureFrameCommand;
+	TSharedPtr<class FUICommandInfo> CaptureFrame;
 };
 
 #endif
diff --git a/Engine/Plugins/Developer/PixWinPlugin/Source/PixWinPlugin/Private/PixWinPluginStyle.cpp b/Engine/Plugins/Developer/PixWinPlugin/Source/PixWinPlugin/Private/PixWinPluginStyle.cpp
--- a/Engine/Plugins/Developer/PixWinPlugin/Source/PixWinPlugin/Private/PixWinPluginStyle.cpp
+++ b/Engine/Plugins/Developer/PixWinPlugin/Source/PixWinPlugin/Private/PixWinPluginStyle.cpp
@@ -44,7 +44,7 @@ void FPixWinPluginStyle::Initialize()
 	}
 
 	StyleSet->Set("PixWinPlugin.Icon", new FSlateImageBrush(FPixWinPluginStyle::InContent("Icon40", ".png"), FVector2D(40.0f, 40.0f)));
-	StyleSet->Set("PixWinPlugin.CaptureFrameIcon", new FSlateImageBrush(FPixWinPluginStyle::InContent("ViewportIcon16", ".png"), FVector2D(16.0f, 16.0f)));
+	StyleSet->Set("PixWinPlugin.CaptureFrame", new FSlateImageBrush(FPixWinPluginStyle::InContent("ViewportIcon16", ".png"), FVector2D(16.0f, 16.0f)));
 
 	FSlateStyleRegistry::RegisterSlateStyle(*StyleSet.Get());
 };
diff --git a/Engine/Plugins/Developer/PixWinPlugin/Source/PixWinPlugin/Private/SPixWinPluginEditorExtension.cpp b/Engine/Plugins/Developer/PixWinPlugin/Source/PixWinPlugin/Private/SPixWinPluginEditorExtension.cpp
--- a/Engine/Plugins/Developer/PixWinPlugin/Source/PixWinPlugin/Private/SPixWinPluginEditorExtension.cpp
+++ b/Engine/Plugins/Developer/PixWinPlugin/Source/PixWinPlugin/Private/SPixWinPluginEditorExtension.cpp
@@ -19,6 +19,8 @@
 #include "Interfaces/IMainFrameModule.h"
 #include "LevelEditor.h"
 #include "RHI.h"
+#include "ToolMenus.h"
+#include "ViewportToolbar/UnrealEdViewportToolbar.h"
 #include "Widgets/Input/SButton.h"
 #include "Widgets/Images/SImage.h"
 
@@ -33,7 +35,7 @@ public:
 	/** Widget constructor */
 	void Construct(const FArguments& Args)
 	{
-		FSlateIcon IconBrush = FSlateIcon(FPixWinPluginStyle::Get()->GetStyleSetName(), "PixWinPlugin.CaptureFrameIcon");
+		FSlateIcon IconBrush = FSlateIcon(FPixWinPluginStyle::Get()->GetStyleSetName(), "PixWinPlugin.CaptureFrame");
 
 		ChildSlot
 		[
@@ -41,10 +43,10 @@ public:
 			.HAlign(HAlign_Center)
 			.VAlign(VAlign_Center)
 			.ContentPadding(FMargin(1.0f))
-			.ToolTipText(FPixWinPluginCommands::Get().CaptureFrameCommand->GetDescription())
+			.ToolTipText(FPixWinPluginCommands::Get().CaptureFrame->GetDescription())
 			.OnClicked_Lambda([this]()
 			{
-				FPlayWorldCommands::GlobalPlayWorldActions->GetActionForCommand(FPixWinPluginCommands::Get().CaptureFrameCommand)->Execute();
+				FPlayWorldCommands::GlobalPlayWorldActions->GetActionForCommand(FPixWinPluginCommands::Get().CaptureFrame)->Execute();
 				return(FReply::Handled());
 			})
 			[
@@ -108,6 +110,8 @@ void FPixWinPluginEditorExtension::Initialize(FPixWinPluginModule* ThePlugin)
 				})
 		);
 		ExtensionManager->AddExtender(ToolbarExtender);
+		
+		ExtendToolbar();
 	}
 #endif // WITH_EDITOR
 
@@ -117,7 +121,7 @@ void FPixWinPluginEditorExtension::Initialize(FPixWinPluginModule* ThePlugin)
 		check(FPlayWorldCommands::GlobalPlayWorldActions.IsValid());
 
 		//Register the editor hotkeys
-		FPlayWorldCommands::GlobalPlayWorldActions->MapAction(FPixWinPluginCommands::Get().CaptureFrameCommand,
+		FPlayWorldCommands::GlobalPlayWorldActions->MapAction(FPixWinPluginCommands::Get().CaptureFrame,
 			FExecuteAction::CreateLambda([]()
 				{
 					FPixWinPluginModule& PluginModule = FModuleManager::GetModuleChecked<FPixWinPluginModule>("PixWinPlugin");
@@ -127,15 +131,39 @@ void FPixWinPluginEditorExtension::Initialize(FPixWinPluginModule* ThePlugin)
 	}
 }
 
+void FPixWinPluginEditorExtension::ExtendToolbar()
+{
+	FToolMenuOwnerScoped ScopedOwner(this);
+
+	UToolMenu* Menu = UToolMenus::Get()->ExtendMenu("LevelEditor.ViewportToolbar");
+	
+	FToolMenuSection& RightSection = Menu->FindOrAddSection("Right");
+	FToolMenuEntry& Entry = RightSection.AddMenuEntry(FPixWinPluginCommands::Get().CaptureFrame);
+	Entry.ToolBarData.LabelOverride = FText::GetEmpty();
+	Entry.InsertPosition.Position = EToolMenuInsertType::First;
+}
+
 void FPixWinPluginEditorExtension::AddToolbarExtension(FToolBarBuilder& ToolbarBuilder, FPixWinPluginModule* ThePlugin)
 {
 #define LOCTEXT_NAMESPACE "LevelEditorToolBar"
 
 	UE_LOG(PixWinPlugin, Verbose, TEXT("Attaching toolbar extension..."));
-	ToolbarBuilder.AddSeparator();
-
-	ToolbarBuilder.BeginSection("PixWinPlugin");
-	ToolbarBuilder.AddWidget(SNew(SPixWinCaptureButton));
+	
+	TAttribute<EVisibility> Visibility;
+	Visibility.BindLambda([]
+	{
+		return UE::UnrealEd::ShowOldViewportToolbars() ? EVisibility::Visible : EVisibility::Collapsed;
+	});
+
+	ToolbarBuilder.BeginSection("PixWinPlugin", false);
+	ToolbarBuilder.AddWidget(
+		SNew(SPixWinCaptureButton),
+		NAME_None,
+		true,
+		HAlign_Fill,
+		FNewMenuDelegate(),
+		Visibility
+	);
 	ToolbarBuilder.EndSection();
 
 #undef LOCTEXT_NAMESPACE
diff --git a/Engine/Plugins/Developer/PixWinPlugin/Source/PixWinPlugin/Private/SPixWinPluginEditorExtension.h b/Engine/Plugins/Developer/PixWinPlugin/Source/PixWinPlugin/Private/SPixWinPluginEditorExtension.h
--- a/Engine/Plugins/Developer/PixWinPlugin/Source/PixWinPlugin/Private/SPixWinPluginEditorExtension.h
+++ b/Engine/Plugins/Developer/PixWinPlugin/Source/PixWinPlugin/Private/SPixWinPluginEditorExtension.h
@@ -20,6 +20,7 @@ public:
 
 private:
 	void Initialize(FPixWinPluginModule* ThePlugin);
+	void ExtendToolbar();
 	void AddToolbarExtension(FToolBarBuilder& ToolbarBuilder, FPixWinPluginModule* ThePlugin);
 
 	TSharedPtr<const FExtensionBase> ToolbarExtension;
diff --git a/Engine/Plugins/Developer/XcodeGPUDebuggerPlugin/Source/XcodeGPUDebuggerPlugin/Private/SXcodeGPUDebuggerPluginEditorExtension.cpp b/Engine/Plugins/Developer/XcodeGPUDebuggerPlugin/Source/XcodeGPUDebuggerPlugin/Private/SXcodeGPUDebuggerPluginEditorExtension.cpp
--- a/Engine/Plugins/Developer/XcodeGPUDebuggerPlugin/Source/XcodeGPUDebuggerPlugin/Private/SXcodeGPUDebuggerPluginEditorExtension.cpp
+++ b/Engine/Plugins/Developer/XcodeGPUDebuggerPlugin/Source/XcodeGPUDebuggerPlugin/Private/SXcodeGPUDebuggerPluginEditorExtension.cpp
@@ -19,6 +19,8 @@
 #include "Framework/MultiBox/MultiBoxDefs.h"
 #include "Interfaces/IMainFrameModule.h"
 #include "LevelEditor.h"
+#include "ToolMenus.h"
+#include "ViewportToolbar/UnrealEdViewportToolbar.h"
 #include "Widgets/Input/SButton.h"
 #include "Widgets/Images/SImage.h"
 
@@ -33,7 +35,7 @@ public:
 	/** Widget constructor */
 	void Construct(const FArguments& Args)
 	{
-		FSlateIcon IconBrush = FSlateIcon(FXcodeGPUDebuggerPluginStyle::Get()->GetStyleSetName(), "XcodeGPUDebuggerPlugin.CaptureFrameIcon");
+		FSlateIcon IconBrush = FSlateIcon(FXcodeGPUDebuggerPluginStyle::Get()->GetStyleSetName(), "XcodeGPUDebuggerPlugin.CaptureFrame");
 
 		ChildSlot
 		[
@@ -42,10 +44,10 @@ public:
 			.VAlign(VAlign_Bottom)
 			.ButtonStyle(FAppStyle::Get(), "ViewportMenu.Button")
 			.ContentPadding(FMargin(1.0f))
-			.ToolTipText(FXcodeGPUDebuggerPluginCommands::Get().CaptureFrameCommand->GetDescription())
+			.ToolTipText(FXcodeGPUDebuggerPluginCommands::Get().CaptureFrame->GetDescription())
 			.OnClicked_Lambda([this]()
 				{
-					FPlayWorldCommands::GlobalPlayWorldActions->GetActionForCommand(FXcodeGPUDebuggerPluginCommands::Get().CaptureFrameCommand)->Execute();
+					FPlayWorldCommands::GlobalPlayWorldActions->GetActionForCommand(FXcodeGPUDebuggerPluginCommands::Get().CaptureFrame)->Execute();
 					return(FReply::Handled());
 				})
 			[
@@ -120,6 +122,8 @@ void FXcodeGPUDebuggerPluginEditorExtension::Initialize(FXcodeGPUDebuggerPluginM
 		AddToolbarExtension(ToolbarBuilder, ThePlugin); })
 	);
 	ExtensionManager->AddExtender(ToolbarExtender);
+	
+	ExtendToolbar();
 
 	IsEditorInitialized = false;
 	FSlateRenderer* SlateRenderer = FSlateApplication::Get().GetRenderer();
@@ -149,7 +153,7 @@ void FXcodeGPUDebuggerPluginEditorExtension::OnEditorLoaded(SWindow& SlateWindow
 	if(FPlayWorldCommands::GlobalPlayWorldActions.IsValid())
 	{
 		//Register the editor hotkeys
-		FPlayWorldCommands::GlobalPlayWorldActions->MapAction(FXcodeGPUDebuggerPluginCommands::Get().CaptureFrameCommand,
+		FPlayWorldCommands::GlobalPlayWorldActions->MapAction(FXcodeGPUDebuggerPluginCommands::Get().CaptureFrame,
 			FExecuteAction::CreateLambda([]()
 			{
 				FXcodeGPUDebuggerPluginModule& PluginModule = FModuleManager::GetModuleChecked<FXcodeGPUDebuggerPluginModule>("XcodeGPUDebuggerPlugin");
@@ -159,15 +163,39 @@ void FXcodeGPUDebuggerPluginEditorExtension::OnEditorLoaded(SWindow& SlateWindow
  	}
 }
 
+void FXcodeGPUDebuggerPluginEditorExtension::ExtendToolbar()
+{
+	FToolMenuOwnerScoped ScopedOwner(this);
+
+	UToolMenu* Menu = UToolMenus::Get()->ExtendMenu("LevelEditor.ViewportToolbar");
+	
+	FToolMenuSection& RightSection = Menu->FindOrAddSection("Right");
+	FToolMenuEntry& Entry = RightSection.AddMenuEntry(FXcodeGPUDebuggerPluginCommands::Get().CaptureFrame);
+	Entry.ToolBarData.LabelOverride = FText::GetEmpty();
+	Entry.InsertPosition.Position = EToolMenuInsertType::First;
+}
+
 void FXcodeGPUDebuggerPluginEditorExtension::AddToolbarExtension(FToolBarBuilder& ToolbarBuilder, FXcodeGPUDebuggerPluginModule* ThePlugin)
 {
 #define LOCTEXT_NAMESPACE "LevelEditorToolBar"
 
 	UE_LOG(XcodeGPUDebuggerPlugin, Verbose, TEXT("Attaching toolbar extension..."));
-	ToolbarBuilder.AddSeparator();
 
-	ToolbarBuilder.BeginSection("XcodeGPUDebuggerPlugin");
-	ToolbarBuilder.AddWidget(SNew(SXcodeGPUDebuggerCaptureButton));
+	TAttribute<EVisibility> Visibility;
+	Visibility.BindLambda([]
+	{
+		return UE::UnrealEd::ShowOldViewportToolbars() ? EVisibility::Visible : EVisibility::Collapsed;
+	});
+
+	ToolbarBuilder.BeginSection("XcodeGPUDebuggerPlugin", false);
+	ToolbarBuilder.AddWidget(
+		SNew(SXcodeGPUDebuggerCaptureButton),
+		NAME_None,
+		true,
+		HAlign_Fill,
+		FNewMenuDelegate(),
+		Visibility
+	);
 	ToolbarBuilder.EndSection();
 
 #undef LOCTEXT_NAMESPACE
diff --git a/Engine/Plugins/Developer/XcodeGPUDebuggerPlugin/Source/XcodeGPUDebuggerPlugin/Private/SXcodeGPUDebuggerPluginEditorExtension.h b/Engine/Plugins/Developer/XcodeGPUDebuggerPlugin/Source/XcodeGPUDebuggerPlugin/Private/SXcodeGPUDebuggerPluginEditorExtension.h
--- a/Engine/Plugins/Developer/XcodeGPUDebuggerPlugin/Source/XcodeGPUDebuggerPlugin/Private/SXcodeGPUDebuggerPluginEditorExtension.h
+++ b/Engine/Plugins/Developer/XcodeGPUDebuggerPlugin/Source/XcodeGPUDebuggerPlugin/Private/SXcodeGPUDebuggerPluginEditorExtension.h
@@ -21,6 +21,7 @@ public:
 private:
 	void Initialize(FXcodeGPUDebuggerPluginModule* ThePlugin);
 	void OnEditorLoaded(SWindow& SlateWindow, void* ViewportRHIPtr);
+	void ExtendToolbar();
 	void AddToolbarExtension(FToolBarBuilder& ToolbarBuilder, FXcodeGPUDebuggerPluginModule* ThePlugin);
 
 	FDelegateHandle LoadedDelegateHandle;
diff --git a/Engine/Plugins/Developer/XcodeGPUDebuggerPlugin/Source/XcodeGPUDebuggerPlugin/Private/XcodeGPUDebuggerPluginCommands.cpp b/Engine/Plugins/Developer/XcodeGPUDebuggerPlugin/Source/XcodeGPUDebuggerPlugin/Private/XcodeGPUDebuggerPluginCommands.cpp
--- a/Engine/Plugins/Developer/XcodeGPUDebuggerPlugin/Source/XcodeGPUDebuggerPlugin/Private/XcodeGPUDebuggerPluginCommands.cpp
+++ b/Engine/Plugins/Developer/XcodeGPUDebuggerPlugin/Source/XcodeGPUDebuggerPlugin/Private/XcodeGPUDebuggerPluginCommands.cpp
@@ -8,7 +8,7 @@
 
 void FXcodeGPUDebuggerPluginCommands::RegisterCommands()
 {
-	UI_COMMAND(CaptureFrameCommand, "Capture Frame", "Captures the next frame and launches Xcode (Shift+E)", EUserInterfaceActionType::Button, FInputChord(EKeys::E, EModifierKey::Shift));
+	UI_COMMAND(CaptureFrame, "Capture Frame", "Captures the next frame and launches Xcode (Shift+E)", EUserInterfaceActionType::Button, FInputChord(EKeys::E, EModifierKey::Shift));
 }
 
 #undef LOCTEXT_NAMESPACE
diff --git a/Engine/Plugins/Developer/XcodeGPUDebuggerPlugin/Source/XcodeGPUDebuggerPlugin/Private/XcodeGPUDebuggerPluginCommands.h b/Engine/Plugins/Developer/XcodeGPUDebuggerPlugin/Source/XcodeGPUDebuggerPlugin/Private/XcodeGPUDebuggerPluginCommands.h
--- a/Engine/Plugins/Developer/XcodeGPUDebuggerPlugin/Source/XcodeGPUDebuggerPlugin/Private/XcodeGPUDebuggerPluginCommands.h
+++ b/Engine/Plugins/Developer/XcodeGPUDebuggerPlugin/Source/XcodeGPUDebuggerPlugin/Private/XcodeGPUDebuggerPluginCommands.h
@@ -17,7 +17,7 @@ public:
 
 	virtual void RegisterCommands() override;
 
-	TSharedPtr<class FUICommandInfo> CaptureFrameCommand;
+	TSharedPtr<class FUICommandInfo> CaptureFrame;
 };
 
 #endif // WITH_EDITOR
diff --git a/Engine/Plugins/Developer/XcodeGPUDebuggerPlugin/Source/XcodeGPUDebuggerPlugin/Private/XcodeGPUDebuggerPluginStyle.cpp b/Engine/Plugins/Developer/XcodeGPUDebuggerPlugin/Source/XcodeGPUDebuggerPlugin/Private/XcodeGPUDebuggerPluginStyle.cpp
--- a/Engine/Plugins/Developer/XcodeGPUDebuggerPlugin/Source/XcodeGPUDebuggerPlugin/Private/XcodeGPUDebuggerPluginStyle.cpp
+++ b/Engine/Plugins/Developer/XcodeGPUDebuggerPlugin/Source/XcodeGPUDebuggerPlugin/Private/XcodeGPUDebuggerPluginStyle.cpp
@@ -44,7 +44,7 @@ void FXcodeGPUDebuggerPluginStyle::Initialize()
 	}
 
 	StyleSet->Set("XcodeGPUDebuggerPlugin.Icon", new FSlateImageBrush(FXcodeGPUDebuggerPluginStyle::InContent("Icon128", ".png"), FVector2D(128.0f, 128.0f)));
-	StyleSet->Set("XcodeGPUDebuggerPlugin.CaptureFrameIcon", new FSlateImageBrush(FXcodeGPUDebuggerPluginStyle::InContent("ViewportIcon16", ".png"), FVector2D(16.0f, 16.0f)));
+	StyleSet->Set("XcodeGPUDebuggerPlugin.CaptureFrame", new FSlateImageBrush(FXcodeGPUDebuggerPluginStyle::InContent("ViewportIcon16", ".png"), FVector2D(16.0f, 16.0f)));
 
 	FSlateStyleRegistry::RegisterSlateStyle(*StyleSet.Get());
 };
diff --git a/Engine/Plugins/Developer/XcodeGPUDebuggerPlugin/Source/XcodeGPUDebuggerPlugin/XcodeGPUDebuggerPlugin.Build.cs b/Engine/Plugins/Developer/XcodeGPUDebuggerPlugin/Source/XcodeGPUDebuggerPlugin/XcodeGPUDebuggerPlugin.Build.cs
--- a/Engine/Plugins/Developer/XcodeGPUDebuggerPlugin/Source/XcodeGPUDebuggerPlugin/XcodeGPUDebuggerPlugin.Build.cs
+++ b/Engine/Plugins/Developer/XcodeGPUDebuggerPlugin/Source/XcodeGPUDebuggerPlugin/XcodeGPUDebuggerPlugin.Build.cs
@@ -37,6 +37,7 @@ namespace UnrealBuildTool.Rules
 					, "UnrealEd"
 					, "MainFrame"
 					, "GameProjectGeneration"
+					, "ToolMenus"
 				});
 			}
 		}
