Update the various GPU debugger extensions to only support the new toolbar.

diff --git a/Engine/Plugins/Developer/PixWinPlugin/Source/PixWinPlugin/Private/SPixWinPluginEditorExtension.cpp b/Engine/Plugins/Developer/PixWinPlugin/Source/PixWinPlugin/Private/SPixWinPluginEditorExtension.cpp
--- a/Engine/Plugins/Developer/PixWinPlugin/Source/PixWinPlugin/Private/SPixWinPluginEditorExtension.cpp
+++ b/Engine/Plugins/Developer/PixWinPlugin/Source/PixWinPlugin/Private/SPixWinPluginEditorExtension.cpp
@@ -26,60 +26,16 @@
 
 extern UNREALED_API UEditorEngine* GEditor;
 
-class SPixWinCaptureButton : public SViewportToolBar
-{
-public:
-	SLATE_BEGIN_ARGS(SPixWinCaptureButton) { }
-	SLATE_END_ARGS()
-
-	/** Widget constructor */
-	void Construct(const FArguments& Args)
-	{
-		FSlateIcon IconBrush = FSlateIcon(FPixWinPluginStyle::Get()->GetStyleSetName(), "PixWinPlugin.CaptureFrame");
-
-		ChildSlot
-		[
-			SNew(SButton)
-			.HAlign(HAlign_Center)
-			.VAlign(VAlign_Center)
-			.ContentPadding(FMargin(1.0f))
-			.ToolTipText(FPixWinPluginCommands::Get().CaptureFrame->GetDescription())
-			.OnClicked_Lambda([this]()
-			{
-				FPlayWorldCommands::GlobalPlayWorldActions->GetActionForCommand(FPixWinPluginCommands::Get().CaptureFrame)->Execute();
-				return(FReply::Handled());
-			})
-			[
-				SNew(SImage)
-				.Image(IconBrush.GetIcon())
-			]
-		];
-	}
-};
-
 FPixWinPluginEditorExtension::FPixWinPluginEditorExtension(FPixWinPluginModule* ThePlugin)
-	: ToolbarExtension()
-	, ExtensionManager()
-	, ToolbarExtender()
 {
 	Initialize(ThePlugin);
 }
 
 FPixWinPluginEditorExtension::~FPixWinPluginEditorExtension()
 {
-	if (ExtensionManager.IsValid())
-	{
-		FPixWinPluginStyle::Shutdown();
-		FPixWinPluginCommands::Unregister();
-
-		ToolbarExtender->RemoveExtension(ToolbarExtension.ToSharedRef());
-
-		ExtensionManager->RemoveExtender(ToolbarExtender);
-	}
-	else
-	{
-		ExtensionManager.Reset();
-	}
+	FPixWinPluginStyle::Shutdown();
+	FPixWinPluginCommands::Unregister();
+	UToolMenus::UnregisterOwner(this);
 }
 
 void FPixWinPluginEditorExtension::Initialize(FPixWinPluginModule* ThePlugin)
@@ -99,19 +55,14 @@ void FPixWinPluginEditorExtension::Initialize(FPixWinPluginModule* ThePlugin)
 #if WITH_EDITOR
 	if (!IsRunningGame())
 	{
-		FLevelEditorModule& LevelEditorModule = FModuleManager::LoadModuleChecked<FLevelEditorModule>("LevelEditor");
-		TSharedRef<FUICommandList> CommandBindings = LevelEditorModule.GetGlobalLevelEditorActions();
-		ExtensionManager = LevelEditorModule.GetToolBarExtensibilityManager();
-		ToolbarExtender = MakeShareable(new FExtender);
-		ToolbarExtension = ToolbarExtender->AddToolBarExtension("CameraSpeed", EExtensionHook::After, CommandBindings,
-			FToolBarExtensionDelegate::CreateLambda([this, ThePlugin](FToolBarBuilder& ToolbarBuilder)
-				{
-					AddToolbarExtension(ToolbarBuilder, ThePlugin);
-				})
-		);
-		ExtensionManager->AddExtender(ToolbarExtender);
-		
-		ExtendToolbar();
+		FToolMenuOwnerScoped ScopedOwner(this);
+
+		UToolMenu* Menu = UToolMenus::Get()->ExtendMenu("LevelEditor.ViewportToolbar");
+	
+		FToolMenuSection& RightSection = Menu->FindOrAddSection("Right");
+		FToolMenuEntry& Entry = RightSection.AddMenuEntry(FPixWinPluginCommands::Get().CaptureFrame);
+		Entry.ToolBarData.LabelOverride = FText::GetEmpty();
+		Entry.InsertPosition.Position = EToolMenuInsertType::First;
 	}
 #endif // WITH_EDITOR
 
@@ -131,42 +82,4 @@ void FPixWinPluginEditorExtension::Initialize(FPixWinPluginModule* ThePlugin)
 	}
 }
 
-void FPixWinPluginEditorExtension::ExtendToolbar()
-{
-	FToolMenuOwnerScoped ScopedOwner(this);
-
-	UToolMenu* Menu = UToolMenus::Get()->ExtendMenu("LevelEditor.ViewportToolbar");
-	
-	FToolMenuSection& RightSection = Menu->FindOrAddSection("Right");
-	FToolMenuEntry& Entry = RightSection.AddMenuEntry(FPixWinPluginCommands::Get().CaptureFrame);
-	Entry.ToolBarData.LabelOverride = FText::GetEmpty();
-	Entry.InsertPosition.Position = EToolMenuInsertType::First;
-}
-
-void FPixWinPluginEditorExtension::AddToolbarExtension(FToolBarBuilder& ToolbarBuilder, FPixWinPluginModule* ThePlugin)
-{
-#define LOCTEXT_NAMESPACE "LevelEditorToolBar"
-
-	UE_LOG(PixWinPlugin, Verbose, TEXT("Attaching toolbar extension..."));
-	
-	TAttribute<EVisibility> Visibility;
-	Visibility.BindLambda([]
-	{
-		return UE::UnrealEd::ShowOldViewportToolbars() ? EVisibility::Visible : EVisibility::Collapsed;
-	});
-
-	ToolbarBuilder.BeginSection("PixWinPlugin", false);
-	ToolbarBuilder.AddWidget(
-		SNew(SPixWinCaptureButton),
-		NAME_None,
-		true,
-		HAlign_Fill,
-		FNewMenuDelegate(),
-		Visibility
-	);
-	ToolbarBuilder.EndSection();
-
-#undef LOCTEXT_NAMESPACE
-}
-
 #endif //WITH_EDITOR
diff --git a/Engine/Plugins/Developer/PixWinPlugin/Source/PixWinPlugin/Private/SPixWinPluginEditorExtension.h b/Engine/Plugins/Developer/PixWinPlugin/Source/PixWinPlugin/Private/SPixWinPluginEditorExtension.h
--- a/Engine/Plugins/Developer/PixWinPlugin/Source/PixWinPlugin/Private/SPixWinPluginEditorExtension.h
+++ b/Engine/Plugins/Developer/PixWinPlugin/Source/PixWinPlugin/Private/SPixWinPluginEditorExtension.h
@@ -20,12 +20,6 @@ public:
 
 private:
 	void Initialize(FPixWinPluginModule* ThePlugin);
-	void ExtendToolbar();
-	void AddToolbarExtension(FToolBarBuilder& ToolbarBuilder, FPixWinPluginModule* ThePlugin);
-
-	TSharedPtr<const FExtensionBase> ToolbarExtension;
-	TSharedPtr<FExtensibilityManager> ExtensionManager;
-	TSharedPtr<FExtender> ToolbarExtender;
 };
 
 #endif //WITH_EDITOR
diff --git a/Engine/Plugins/Developer/RenderDocPlugin/Source/RenderDocPlugin/Private/SRenderDocPluginEditorExtension.cpp b/Engine/Plugins/Developer/RenderDocPlugin/Source/RenderDocPlugin/Private/SRenderDocPluginEditorExtension.cpp
--- a/Engine/Plugins/Developer/RenderDocPlugin/Source/RenderDocPlugin/Private/SRenderDocPluginEditorExtension.cpp
+++ b/Engine/Plugins/Developer/RenderDocPlugin/Source/RenderDocPlugin/Private/SRenderDocPluginEditorExtension.cpp
@@ -28,61 +28,15 @@
 
 extern UNREALED_API UEditorEngine* GEditor;
 
-class SRenderDocCaptureButton : public SViewportToolBar
-{
-public:
-	SLATE_BEGIN_ARGS(SRenderDocCaptureButton) { }
-	SLATE_END_ARGS()
-
-	/** Widget constructor */
-	void Construct(const FArguments& Args)
-	{
-		FSlateIcon IconBrush = FSlateIcon(FRenderDocPluginStyle::Get()->GetStyleSetName(), "RenderDocPlugin.CaptureFrame");
-		
-		ChildSlot
-		[
-			SNew(SButton)
-			.HAlign(HAlign_Center)
-			.VAlign(VAlign_Bottom)
-			.ContentPadding(FMargin(1.0f))
-			.ToolTipText(FRenderDocPluginCommands::Get().CaptureFrame->GetDescription())
-			.OnClicked_Lambda([this]() 
-				{ 
-					FPlayWorldCommands::GlobalPlayWorldActions->GetActionForCommand(FRenderDocPluginCommands::Get().CaptureFrame)->Execute();
-					return(FReply::Handled()); 
-				})
-			[
-				SNew(SImage)
-				.Image(IconBrush.GetIcon())
-			]
-		];
-	}
-};
-
 FRenderDocPluginEditorExtension::FRenderDocPluginEditorExtension(FRenderDocPluginModule* ThePlugin)
-	: ToolbarExtension()
-	, ExtensionManager()
-	, ToolbarExtender()
 {
 	Initialize(ThePlugin);
 }
 
 FRenderDocPluginEditorExtension::~FRenderDocPluginEditorExtension()
 {
-	if (ExtensionManager.IsValid())
-	{
-		FRenderDocPluginStyle::Shutdown();
-		FRenderDocPluginCommands::Unregister();
-
-		ToolbarExtender->RemoveExtension(ToolbarExtension.ToSharedRef());
-
-		ExtensionManager->RemoveExtender(ToolbarExtender);
-	}
-	else
-	{
-		ExtensionManager.Reset();
-	}
-	
+	FRenderDocPluginStyle::Shutdown();
+	FRenderDocPluginCommands::Unregister();
 	UToolMenus::UnregisterOwner(this);
 }
 
@@ -103,19 +57,14 @@ void FRenderDocPluginEditorExtension::Initialize(FRenderDocPluginModule* ThePlug
 #if WITH_EDITOR
 	if (!IsRunningGame())
 	{
-		FLevelEditorModule& LevelEditorModule = FModuleManager::LoadModuleChecked<FLevelEditorModule>("LevelEditor");
-		TSharedRef<FUICommandList> CommandBindings = LevelEditorModule.GetGlobalLevelEditorActions();
-		ExtensionManager = LevelEditorModule.GetToolBarExtensibilityManager();
-		ToolbarExtender = MakeShareable(new FExtender);
-		ToolbarExtension = ToolbarExtender->AddToolBarExtension("CameraSpeed", EExtensionHook::After, CommandBindings,
-			FToolBarExtensionDelegate::CreateLambda([this, ThePlugin](FToolBarBuilder& ToolbarBuilder)
-				{
-					AddToolbarExtension(ToolbarBuilder, ThePlugin); 
-				})
-		);
-		ExtensionManager->AddExtender(ToolbarExtender);
-		
-		ExtendToolbar();
+		FToolMenuOwnerScoped ScopedOwner(this);
+
+		UToolMenu* Menu = UToolMenus::Get()->ExtendMenu("LevelEditor.ViewportToolbar");
+	
+		FToolMenuSection& RightSection = Menu->FindOrAddSection("Right");
+		FToolMenuEntry& Entry = RightSection.AddMenuEntry(FRenderDocPluginCommands::Get().CaptureFrame);
+		Entry.ToolBarData.LabelOverride = FText::GetEmpty();
+		Entry.InsertPosition.Position = EToolMenuInsertType::First;
 	}
 #endif // WITH_EDITOR
 
@@ -141,44 +90,4 @@ void FRenderDocPluginEditorExtension::Initialize(FRenderDocPluginModule* ThePlug
 	}
 }
 
-void FRenderDocPluginEditorExtension::ExtendToolbar()
-{
-	FToolMenuOwnerScoped ScopedOwner(this);
-
-	UToolMenu* Menu = UToolMenus::Get()->ExtendMenu("LevelEditor.ViewportToolbar");
-	
-	FToolMenuSection& RightSection = Menu->FindOrAddSection("Right");
-	FToolMenuEntry& Entry = RightSection.AddMenuEntry(FRenderDocPluginCommands::Get().CaptureFrame);
-	Entry.ToolBarData.LabelOverride = FText::GetEmpty();
-	Entry.InsertPosition.Position = EToolMenuInsertType::First;
-}
-
-void FRenderDocPluginEditorExtension::AddToolbarExtension(FToolBarBuilder& ToolbarBuilder, FRenderDocPluginModule* ThePlugin)
-{
-#define LOCTEXT_NAMESPACE "LevelEditorToolBar"
-
-	UE_LOG(RenderDocPlugin, Verbose, TEXT("Attaching toolbar extension..."));
-	ToolbarBuilder.BeginSection("RenderdocPlugin", false);
-	
-	TAttribute<EVisibility> Visibility;
-	Visibility.BindLambda([]
-	{
-		return UE::UnrealEd::ShowOldViewportToolbars() ? EVisibility::Visible : EVisibility::Collapsed;
-	});
-	
-	ToolbarBuilder.AddSeparator(NAME_None, Visibility);
-	
-	ToolbarBuilder.AddWidget(
-		SNew(SRenderDocCaptureButton),
-		NAME_None,
-		true,
-		HAlign_Fill,
-		FNewMenuDelegate(),
-		Visibility
-	);
-	ToolbarBuilder.EndSection();
-
-#undef LOCTEXT_NAMESPACE
-}
-
 #endif //WITH_EDITOR
diff --git a/Engine/Plugins/Developer/RenderDocPlugin/Source/RenderDocPlugin/Private/SRenderDocPluginEditorExtension.h b/Engine/Plugins/Developer/RenderDocPlugin/Source/RenderDocPlugin/Private/SRenderDocPluginEditorExtension.h
--- a/Engine/Plugins/Developer/RenderDocPlugin/Source/RenderDocPlugin/Private/SRenderDocPluginEditorExtension.h
+++ b/Engine/Plugins/Developer/RenderDocPlugin/Source/RenderDocPlugin/Private/SRenderDocPluginEditorExtension.h
@@ -20,12 +20,6 @@ public:
 
 private:
   void Initialize(FRenderDocPluginModule* ThePlugin);
-  void ExtendToolbar();
-  void AddToolbarExtension(FToolBarBuilder& ToolbarBuilder, FRenderDocPluginModule* ThePlugin);
-
-  TSharedPtr<const FExtensionBase> ToolbarExtension;
-  TSharedPtr<FExtensibilityManager> ExtensionManager;
-  TSharedPtr<FExtender> ToolbarExtender;
 };
 
 #endif //WITH_EDITOR
diff --git a/Engine/Plugins/Developer/XcodeGPUDebuggerPlugin/Source/XcodeGPUDebuggerPlugin/Private/SXcodeGPUDebuggerPluginEditorExtension.cpp b/Engine/Plugins/Developer/XcodeGPUDebuggerPlugin/Source/XcodeGPUDebuggerPlugin/Private/SXcodeGPUDebuggerPluginEditorExtension.cpp
--- a/Engine/Plugins/Developer/XcodeGPUDebuggerPlugin/Source/XcodeGPUDebuggerPlugin/Private/SXcodeGPUDebuggerPluginEditorExtension.cpp
+++ b/Engine/Plugins/Developer/XcodeGPUDebuggerPlugin/Source/XcodeGPUDebuggerPlugin/Private/SXcodeGPUDebuggerPluginEditorExtension.cpp
@@ -26,43 +26,8 @@
 
 extern UNREALED_API UEditorEngine* GEditor;
 
-class SXcodeGPUDebuggerCaptureButton : public SViewportToolBar
-{
-public:
-	SLATE_BEGIN_ARGS(SXcodeGPUDebuggerCaptureButton) { }
-	SLATE_END_ARGS()
-
-	/** Widget constructor */
-	void Construct(const FArguments& Args)
-	{
-		FSlateIcon IconBrush = FSlateIcon(FXcodeGPUDebuggerPluginStyle::Get()->GetStyleSetName(), "XcodeGPUDebuggerPlugin.CaptureFrame");
-
-		ChildSlot
-		[
-			SNew(SButton)
-			.HAlign(HAlign_Center)
-			.VAlign(VAlign_Bottom)
-			.ButtonStyle(FAppStyle::Get(), "ViewportMenu.Button")
-			.ContentPadding(FMargin(1.0f))
-			.ToolTipText(FXcodeGPUDebuggerPluginCommands::Get().CaptureFrame->GetDescription())
-			.OnClicked_Lambda([this]()
-				{
-					FPlayWorldCommands::GlobalPlayWorldActions->GetActionForCommand(FXcodeGPUDebuggerPluginCommands::Get().CaptureFrame)->Execute();
-					return(FReply::Handled());
-				})
-			[
-				SNew(SImage)
-				.Image(IconBrush.GetIcon())
-			]
-		];
-	}
-};
-
 FXcodeGPUDebuggerPluginEditorExtension::FXcodeGPUDebuggerPluginEditorExtension(FXcodeGPUDebuggerPluginModule* ThePlugin)
 	: LoadedDelegateHandle()
-	, ToolbarExtension()
-	, ExtensionManager()
-	, ToolbarExtender()
 	, IsEditorInitialized(false)
 {
 	// Defer Level Editor UI extensions until Level Editor has been loaded:
@@ -84,19 +49,9 @@ FXcodeGPUDebuggerPluginEditorExtension::FXcodeGPUDebuggerPluginEditorExtension(F
 
 FXcodeGPUDebuggerPluginEditorExtension::~FXcodeGPUDebuggerPluginEditorExtension()
 {
-	if (ExtensionManager.IsValid())
-	{
-		FXcodeGPUDebuggerPluginStyle::Shutdown();
-		FXcodeGPUDebuggerPluginCommands::Unregister();
-
-		ToolbarExtender->RemoveExtension(ToolbarExtension.ToSharedRef());
-
-		ExtensionManager->RemoveExtender(ToolbarExtender);
-	}
-	else
-	{
-		ExtensionManager.Reset();
-	}
+	FXcodeGPUDebuggerPluginStyle::Shutdown();
+	FXcodeGPUDebuggerPluginCommands::Unregister();
+	UToolMenus::UnregisterOwner(this);
 }
 
 void FXcodeGPUDebuggerPluginEditorExtension::Initialize(FXcodeGPUDebuggerPluginModule* ThePlugin)
@@ -112,18 +67,14 @@ void FXcodeGPUDebuggerPluginEditorExtension::Initialize(FXcodeGPUDebuggerPluginM
 	FXcodeGPUDebuggerPluginStyle::Initialize();
 	FXcodeGPUDebuggerPluginCommands::Register();
 
-	FLevelEditorModule& LevelEditorModule = FModuleManager::LoadModuleChecked<FLevelEditorModule>("LevelEditor");
-	TSharedRef<FUICommandList> CommandBindings = LevelEditorModule.GetGlobalLevelEditorActions();
-	ExtensionManager = LevelEditorModule.GetToolBarExtensibilityManager();
-	ToolbarExtender = MakeShareable(new FExtender);
-	ToolbarExtension = ToolbarExtender->AddToolBarExtension("CameraSpeed", EExtensionHook::After, CommandBindings,
-		FToolBarExtensionDelegate::CreateLambda([this, ThePlugin](FToolBarBuilder& ToolbarBuilder)
-	{
-		AddToolbarExtension(ToolbarBuilder, ThePlugin); })
-	);
-	ExtensionManager->AddExtender(ToolbarExtender);
+	FToolMenuOwnerScoped ScopedOwner(this);
+
+	UToolMenu* Menu = UToolMenus::Get()->ExtendMenu("LevelEditor.ViewportToolbar");
 	
-	ExtendToolbar();
+	FToolMenuSection& RightSection = Menu->FindOrAddSection("Right");
+	FToolMenuEntry& Entry = RightSection.AddMenuEntry(FXcodeGPUDebuggerPluginCommands::Get().CaptureFrame);
+	Entry.ToolBarData.LabelOverride = FText::GetEmpty();
+	Entry.InsertPosition.Position = EToolMenuInsertType::First;
 
 	IsEditorInitialized = false;
 	FSlateRenderer* SlateRenderer = FSlateApplication::Get().GetRenderer();
@@ -163,42 +114,4 @@ void FXcodeGPUDebuggerPluginEditorExtension::OnEditorLoaded(SWindow& SlateWindow
  	}
 }
 
-void FXcodeGPUDebuggerPluginEditorExtension::ExtendToolbar()
-{
-	FToolMenuOwnerScoped ScopedOwner(this);
-
-	UToolMenu* Menu = UToolMenus::Get()->ExtendMenu("LevelEditor.ViewportToolbar");
-	
-	FToolMenuSection& RightSection = Menu->FindOrAddSection("Right");
-	FToolMenuEntry& Entry = RightSection.AddMenuEntry(FXcodeGPUDebuggerPluginCommands::Get().CaptureFrame);
-	Entry.ToolBarData.LabelOverride = FText::GetEmpty();
-	Entry.InsertPosition.Position = EToolMenuInsertType::First;
-}
-
-void FXcodeGPUDebuggerPluginEditorExtension::AddToolbarExtension(FToolBarBuilder& ToolbarBuilder, FXcodeGPUDebuggerPluginModule* ThePlugin)
-{
-#define LOCTEXT_NAMESPACE "LevelEditorToolBar"
-
-	UE_LOG(XcodeGPUDebuggerPlugin, Verbose, TEXT("Attaching toolbar extension..."));
-
-	TAttribute<EVisibility> Visibility;
-	Visibility.BindLambda([]
-	{
-		return UE::UnrealEd::ShowOldViewportToolbars() ? EVisibility::Visible : EVisibility::Collapsed;
-	});
-
-	ToolbarBuilder.BeginSection("XcodeGPUDebuggerPlugin", false);
-	ToolbarBuilder.AddWidget(
-		SNew(SXcodeGPUDebuggerCaptureButton),
-		NAME_None,
-		true,
-		HAlign_Fill,
-		FNewMenuDelegate(),
-		Visibility
-	);
-	ToolbarBuilder.EndSection();
-
-#undef LOCTEXT_NAMESPACE
-}
-
 #endif //WITH_EDITOR
diff --git a/Engine/Plugins/Developer/XcodeGPUDebuggerPlugin/Source/XcodeGPUDebuggerPlugin/Private/SXcodeGPUDebuggerPluginEditorExtension.h b/Engine/Plugins/Developer/XcodeGPUDebuggerPlugin/Source/XcodeGPUDebuggerPlugin/Private/SXcodeGPUDebuggerPluginEditorExtension.h
--- a/Engine/Plugins/Developer/XcodeGPUDebuggerPlugin/Source/XcodeGPUDebuggerPlugin/Private/SXcodeGPUDebuggerPluginEditorExtension.h
+++ b/Engine/Plugins/Developer/XcodeGPUDebuggerPlugin/Source/XcodeGPUDebuggerPlugin/Private/SXcodeGPUDebuggerPluginEditorExtension.h
@@ -20,14 +20,8 @@ public:
 
 private:
 	void Initialize(FXcodeGPUDebuggerPluginModule* ThePlugin);
-	void OnEditorLoaded(SWindow& SlateWindow, void* ViewportRHIPtr);
-	void ExtendToolbar();
-	void AddToolbarExtension(FToolBarBuilder& ToolbarBuilder, FXcodeGPUDebuggerPluginModule* ThePlugin);
 
 	FDelegateHandle LoadedDelegateHandle;
-	TSharedPtr<const FExtensionBase> ToolbarExtension;
-	TSharedPtr<FExtensibilityManager> ExtensionManager;
-	TSharedPtr<FExtender> ToolbarExtender;
 	bool IsEditorInitialized;
 };
 
