Fix an issue submitted in 43006398

#rnx

diff --git a/Engine/Plugins/Developer/XcodeGPUDebuggerPlugin/Source/XcodeGPUDebuggerPlugin/Private/SXcodeGPUDebuggerPluginEditorExtension.h b/Engine/Plugins/Developer/XcodeGPUDebuggerPlugin/Source/XcodeGPUDebuggerPlugin/Private/SXcodeGPUDebuggerPluginEditorExtension.h
--- a/Engine/Plugins/Developer/XcodeGPUDebuggerPlugin/Source/XcodeGPUDebuggerPlugin/Private/SXcodeGPUDebuggerPluginEditorExtension.h
+++ b/Engine/Plugins/Developer/XcodeGPUDebuggerPlugin/Source/XcodeGPUDebuggerPlugin/Private/SXcodeGPUDebuggerPluginEditorExtension.h
@@ -20,6 +20,7 @@ public:
 
 private:
 	void Initialize(FXcodeGPUDebuggerPluginModule* ThePlugin);
+	void OnEditorLoaded(SWindow& SlateWindow, void* ViewportRHIPtr);
 
 	FDelegateHandle LoadedDelegateHandle;
 	bool IsEditorInitialized;
