[CaptureFrame] Give the RenderDoc, PixWin, and Xcode "CaptureFrame" toolbar entries each their own unique id.

[RN] Minor | Editor
#Synced-CL 51758956

diff --git a/Engine/Plugins/Developer/PixWinPlugin/Source/PixWinPlugin/Private/SPixWinPluginEditorExtension.cpp b/Engine/Plugins/Developer/PixWinPlugin/Source/PixWinPlugin/Private/SPixWinPluginEditorExtension.cpp
--- a/Engine/Plugins/Developer/PixWinPlugin/Source/PixWinPlugin/Private/SPixWinPluginEditorExtension.cpp
+++ b/Engine/Plugins/Developer/PixWinPlugin/Source/PixWinPlugin/Private/SPixWinPluginEditorExtension.cpp
@@ -60,7 +60,7 @@ void FPixWinPluginEditorExtension::Initialize(FPixWinPluginModule* ThePlugin)
 		UToolMenu* Menu = UToolMenus::Get()->ExtendMenu("LevelEditor.ViewportToolbar");
 	
 		FToolMenuSection& RightSection = Menu->FindOrAddSection("Right");
-		FToolMenuEntry& Entry = RightSection.AddMenuEntry(FPixWinPluginCommands::Get().CaptureFrame);
+		FToolMenuEntry& Entry = RightSection.AddMenuEntry("PixWinCaptureFrame", FPixWinPluginCommands::Get().CaptureFrame);
 		Entry.ToolBarData.LabelOverride = FText::GetEmpty();
 		Entry.InsertPosition.Position = EToolMenuInsertType::First;
 	}
diff --git a/Engine/Plugins/Developer/RenderDocPlugin/Source/RenderDocPlugin/Private/SRenderDocPluginEditorExtension.cpp b/Engine/Plugins/Developer/RenderDocPlugin/Source/RenderDocPlugin/Private/SRenderDocPluginEditorExtension.cpp
--- a/Engine/Plugins/Developer/RenderDocPlugin/Source/RenderDocPlugin/Private/SRenderDocPluginEditorExtension.cpp
+++ b/Engine/Plugins/Developer/RenderDocPlugin/Source/RenderDocPlugin/Private/SRenderDocPluginEditorExtension.cpp
@@ -62,7 +62,7 @@ void FRenderDocPluginEditorExtension::Initialize(FRenderDocPluginModule* ThePlug
 		UToolMenu* Menu = UToolMenus::Get()->ExtendMenu("LevelEditor.ViewportToolbar");
 	
 		FToolMenuSection& RightSection = Menu->FindOrAddSection("Right");
-		FToolMenuEntry& Entry = RightSection.AddMenuEntry(FRenderDocPluginCommands::Get().CaptureFrame);
+		FToolMenuEntry& Entry = RightSection.AddMenuEntry("RenderDocCaptureFrame", FRenderDocPluginCommands::Get().CaptureFrame);
 		Entry.ToolBarData.LabelOverride = FText::GetEmpty();
 		Entry.InsertPosition.Position = EToolMenuInsertType::First;
 	}
diff --git a/Engine/Plugins/Developer/XcodeGPUDebuggerPlugin/Source/XcodeGPUDebuggerPlugin/Private/SXcodeGPUDebuggerPluginEditorExtension.cpp b/Engine/Plugins/Developer/XcodeGPUDebuggerPlugin/Source/XcodeGPUDebuggerPlugin/Private/SXcodeGPUDebuggerPluginEditorExtension.cpp
--- a/Engine/Plugins/Developer/XcodeGPUDebuggerPlugin/Source/XcodeGPUDebuggerPlugin/Private/SXcodeGPUDebuggerPluginEditorExtension.cpp
+++ b/Engine/Plugins/Developer/XcodeGPUDebuggerPlugin/Source/XcodeGPUDebuggerPlugin/Private/SXcodeGPUDebuggerPluginEditorExtension.cpp
@@ -72,7 +72,7 @@ void FXcodeGPUDebuggerPluginEditorExtension::Initialize(FXcodeGPUDebuggerPluginM
 	UToolMenu* Menu = UToolMenus::Get()->ExtendMenu("LevelEditor.ViewportToolbar");
 	
 	FToolMenuSection& RightSection = Menu->FindOrAddSection("Right");
-	FToolMenuEntry& Entry = RightSection.AddMenuEntry(FXcodeGPUDebuggerPluginCommands::Get().CaptureFrame);
+	FToolMenuEntry& Entry = RightSection.AddMenuEntry("XCodeCaptureFrame", FXcodeGPUDebuggerPluginCommands::Get().CaptureFrame);
 	Entry.ToolBarData.LabelOverride = FText::GetEmpty();
 	Entry.InsertPosition.Position = EToolMenuInsertType::First;
 
